Copied CL #6953827 from //UE4/Dev-Enterprise/...

- Revisited architecture around dataprep asset: Now, the dataprep asset performs all operations on its properties. Components must listen to its events to react to its changes, i.e. dataprep asset view, dataprep editor, etc.
- Added code to identify producers which supersede others

diff --git a/Engine/Plugins/Experimental/DataPrepEditor/Source/DataPrepCore/Private/DataPrepAsset.cpp b/Engine/Plugins/Experimental/DataPrepEditor/Source/DataPrepCore/Private/DataPrepAsset.cpp
--- a/Engine/Plugins/Experimental/DataPrepEditor/Source/DataPrepCore/Private/DataPrepAsset.cpp
+++ b/Engine/Plugins/Experimental/DataPrepEditor/Source/DataPrepCore/Private/DataPrepAsset.cpp
@@ -5,10 +5,30 @@
 #include "DataprepActionAsset.h"
 #include "DataprepCoreLogCategory.h"
 #include "DataprepCoreUtils.h"
+#include "DataPrepRecipe.h"
 
+#include "AssetRegistryModule.h"
 #ifdef WITH_EDITOR
 #include "Editor.h"
 #endif //WITH_EDITOR
+#include "Kismet2/KismetEditorUtilities.h"
+
+namespace DataprepAssetUtil
+{
+	void DeleteRegisteredAsset(UObject* Asset)
+	{
+		if(Asset != nullptr)
+		{
+			Asset->Rename( nullptr, GetTransientPackage(), REN_DontCreateRedirectors | REN_NonTransactional );
+
+			Asset->ClearFlags(RF_Standalone | RF_Public);
+			Asset->RemoveFromRoot();
+			Asset->MarkPendingKill();
+
+			FAssetRegistryModule::AssetDeleted( Asset ) ;
+		}
+	}
+}
 
 // FDataprepAssetAction =================================================================
 
@@ -50,15 +70,15 @@ void FDataprepAssetAction::BindDataprepAssetToAction()
 {
 	if ( ActionAsset )
 	{
-		OnOperationOrderChandedHandle = ActionAsset->GetOnStepsOrderChanged().AddRaw( this, &FDataprepAssetAction::OnActionOperationsOrderChanged );
+		OnOperationOrderChangedHandle = ActionAsset->GetOnStepsOrderChanged().AddRaw( this, &FDataprepAssetAction::OnActionOperationsOrderChanged );
 	}
 }
 
 void FDataprepAssetAction::UnbindDataprepAssetFromAction()
 {
-	if ( ActionAsset && OnOperationOrderChandedHandle.IsValid() )
+	if ( ActionAsset && OnOperationOrderChangedHandle.IsValid() )
 	{
-		ActionAsset->GetOnStepsOrderChanged().Remove( OnOperationOrderChandedHandle );
+		ActionAsset->GetOnStepsOrderChanged().Remove( OnOperationOrderChangedHandle );
 	}
 }
 
@@ -106,6 +126,67 @@ UDataprepAsset::~UDataprepAsset()
 #endif //WITH_EDITOR
 }
 
+void UDataprepAsset::PostInitProperties()
+{
+	Super::PostInitProperties();
+
+	if(HasAnyFlags(RF_ClassDefaultObject | RF_NeedLoad) == false)
+	{
+		// Set DataprepAsset's consumer to the first registered consumer
+		for( TObjectIterator< UClass > It ; It ; ++It )
+		{
+			UClass* CurrentClass = (*It);
+
+			if ( !CurrentClass->HasAnyClassFlags( CLASS_Abstract ) )
+			{
+				if( CurrentClass->IsChildOf( UDataprepContentConsumer::StaticClass() ) )
+				{
+					Consumer = NewObject< UDataprepContentConsumer >( GetOutermost(), CurrentClass, NAME_None, RF_Transactional );
+					check( Consumer );
+
+					FAssetRegistryModule::AssetCreated( Consumer );
+					Consumer->MarkPackageDirty();
+
+					break;
+				}
+			}
+		}
+
+		// Begin: Temp code for the nodes development
+		const FString DesiredName = GetName() + TEXT("_Recipe");
+		FName BlueprintName = MakeUniqueObjectName( GetOutermost(), UBlueprint::StaticClass(), *DesiredName );
+
+		DataprepRecipeBP = FKismetEditorUtilities::CreateBlueprint( UDataprepRecipe::StaticClass(), this, BlueprintName, BPTYPE_Normal, UBlueprint::StaticClass(), UBlueprintGeneratedClass::StaticClass() );
+		check( DataprepRecipeBP );
+
+		FAssetRegistryModule::AssetCreated( DataprepRecipeBP );
+
+		DataprepRecipeBP->MarkPackageDirty();
+
+		DataprepRecipeBP->OnChanged().AddUObject( this, &UDataprepAsset::OnBlueprintChanged );
+		// End: Temp code for the nodes development
+	}
+}
+
+void UDataprepAsset::Serialize( FArchive& Ar )
+{
+	Super::Serialize( Ar );
+
+	if( Ar.IsLoading() )
+	{
+		check(Consumer);
+		Consumer->GetOnChanged().AddUObject( this, &UDataprepAsset::OnConsumerChanged );
+
+		check(DataprepRecipeBP);
+		DataprepRecipeBP->OnChanged().AddUObject( this, &UDataprepAsset::OnBlueprintChanged );
+
+		for( FDataprepAssetProducer& Producer : Producers )
+		{
+			Producer.Producer->GetOnChanged().AddUObject( this, &UDataprepAsset::OnProducerChanged );
+		}
+	}
+}
+
 int32 UDataprepAsset::AddAction()
 {
 	UDataprepActionAsset* Action = NewObject< UDataprepActionAsset >( this );
@@ -195,16 +276,6 @@ bool UDataprepAsset::RemoveAction(int32 Index)
 	return false;
 }
 
-FOnStepsOrderChanged& UDataprepAsset::GetOnActionsOrderChanged()
-{
-	return OnActionsOrderChanged;
-}
-
-FOnActionOperationsOrderChanged& UDataprepAsset::GetOnActionOperationsOrderChanged()
-{
-	return OnActionOperationsOrderChanged;
-}
-
 void UDataprepAsset::RemoveInvalidActions()
 {
 	bool bWasActionsModified = false;
@@ -231,7 +302,7 @@ void UDataprepAsset::RunProducers(const UDataprepContentProducer::ProducerContex
 
 	for ( FDataprepAssetProducer& AssetProducer : Producers )
 	{
-		if (AssetProducer.Producer && AssetProducer.bIsEnabled)
+		if (AssetProducer.Producer && ( AssetProducer.bIsEnabled && AssetProducer.SupersededBy == INDEX_NONE ) )
 		{
 			FString OutReason;
 			if (AssetProducer.Producer->Initialize( InContext, OutReason ))
@@ -259,3 +330,244 @@ void UDataprepAsset::RunProducers(const UDataprepContentProducer::ProducerContex
 		}
 	}
 }
+
+bool UDataprepAsset::RunConsumer( const UDataprepContentConsumer::ConsumerContext& InContext, FString& OutReason)
+{
+	if( !Consumer->Initialize( InContext, OutReason ) )
+	{
+		return false;
+	}
+
+	// #ueent_todo: Update state of entry: finalizing
+
+	if ( !Consumer->Run() )
+	{
+		// #ueent_todo: Inform execution has failed
+		return false;
+	}
+
+	Consumer->Reset();
+	
+	return true;
+}
+
+bool UDataprepAsset::AddProducer(UClass* ProducerClass)
+{
+	if( ProducerClass && ProducerClass->IsChildOf( UDataprepContentProducer::StaticClass() ) )
+	{
+		UDataprepContentProducer* Producer = NewObject< UDataprepContentProducer >( GetOutermost(), ProducerClass, NAME_None, RF_Transactional );
+		FAssetRegistryModule::AssetCreated( Producer );
+		Producer->MarkPackageDirty();
+
+		int32 ProducerNextIndex = Producers.Num();
+		Producers.Emplace( Producer, true );
+
+		Producer->GetOnChanged().AddUObject(this, &UDataprepAsset::OnProducerChanged);
+		MarkPackageDirty();
+
+		OnChanged.Broadcast( FDataprepAssetChangeType::ProducerAdded, ProducerNextIndex );
+
+		return true;
+	}
+
+	return false;
+}
+
+bool UDataprepAsset::RemoveProducer(int32 IndexToRemove)
+{
+	if( Producers.IsValidIndex( IndexToRemove ) )
+	{
+		UDataprepContentProducer* Producer = Producers[IndexToRemove].Producer;
+
+		Producer->GetOnChanged().RemoveAll( this );
+
+		DataprepAssetUtil::DeleteRegisteredAsset( Producer );
+
+		Producers.RemoveAt( IndexToRemove );
+
+		if(Producers.Num() == 1)
+		{
+			Producers[0].SupersededBy = INDEX_NONE;
+		}
+		else if(Producers.Num() > 1)
+		{
+			// Update value stored in SupersededBy property where applicable
+			for( FDataprepAssetProducer& AssetProducer : Producers )
+			{
+				if( AssetProducer.SupersededBy == IndexToRemove )
+				{
+					AssetProducer.SupersededBy = INDEX_NONE;
+				}
+				else if( AssetProducer.SupersededBy > IndexToRemove )
+				{
+					--AssetProducer.SupersededBy;
+				}
+			}
+		}
+
+		MarkPackageDirty();
+
+		OnChanged.Broadcast( FDataprepAssetChangeType::ProducerRemoved, IndexToRemove );
+
+		return true;
+	}
+
+	return false;
+}
+
+void UDataprepAsset::EnableProducer(int32 Index, bool bValue)
+{
+	if( Producers.IsValidIndex( Index ) )
+	{
+		Producers[Index].bIsEnabled = bValue;
+
+		MarkPackageDirty();
+
+		// Relay change notification to observers of this object
+		OnChanged.Broadcast( FDataprepAssetChangeType::ProducerModified, Index );
+	}
+}
+
+bool UDataprepAsset::EnableAllProducers(bool bValue)
+{
+	if( Producers.Num() > 0 )
+	{
+		for( FDataprepAssetProducer& Producer : Producers )
+		{
+			Producer.bIsEnabled = bValue;
+		}
+
+		MarkPackageDirty();
+
+		OnChanged.Broadcast( FDataprepAssetChangeType::ProducerModified, INDEX_NONE );
+
+		return true;
+	}
+
+	return false;
+}
+
+bool UDataprepAsset::ReplaceConsumer(UClass* NewConsumerClass)
+{
+	if( NewConsumerClass && NewConsumerClass->IsChildOf(UDataprepContentConsumer::StaticClass()))
+	{
+		if(Consumer != nullptr)
+		{
+			Consumer->GetOnChanged().RemoveAll( this );
+			DataprepAssetUtil::DeleteRegisteredAsset( Consumer );
+		}
+
+		Consumer = NewObject< UDataprepContentConsumer >( GetOutermost(), NewConsumerClass, NAME_None, RF_Transactional );
+		check( Consumer );
+
+		FAssetRegistryModule::AssetCreated( Consumer );
+		Consumer->MarkPackageDirty();
+
+		Consumer->GetOnChanged().AddUObject( this, &UDataprepAsset::OnConsumerChanged );
+		MarkPackageDirty();
+
+		OnChanged.Broadcast( FDataprepAssetChangeType::ConsumerModified, INDEX_NONE );
+
+		return true;
+	}
+
+	return false;
+}
+
+void UDataprepAsset::OnConsumerChanged()
+{
+	MarkPackageDirty();
+
+	// Broadcast change on consumer to observers of this object
+	OnChanged.Broadcast( FDataprepAssetChangeType::ConsumerModified, INDEX_NONE );
+}
+
+void UDataprepAsset::OnProducerChanged( const UDataprepContentProducer* InProducer )
+{
+	int32 FoundIndex = 0;
+	for( FDataprepAssetProducer& AssetProducer : Producers )
+	{
+		if( AssetProducer.Producer == InProducer )
+		{
+			break;
+		}
+
+		++FoundIndex;
+	}
+
+	// Verify found producer is not now superseded by another one
+	if( FoundIndex < Producers.Num() )
+	{
+		bool bChangeAll = false;
+		ValidateProducerChanges( FoundIndex, bChangeAll );
+
+		MarkPackageDirty();
+
+		// Relay change notification to observers of this object
+		OnChanged.Broadcast( FDataprepAssetChangeType::ProducerModified, bChangeAll ? INDEX_NONE : FoundIndex );
+	}
+}
+
+void UDataprepAsset::OnBlueprintChanged( UBlueprint* InBlueprint )
+{
+	if(InBlueprint == DataprepRecipeBP)
+	{
+		OnChanged.Broadcast( FDataprepAssetChangeType::BlueprintModified, INDEX_NONE );
+	}
+}
+
+void UDataprepAsset::ValidateProducerChanges( int32 Index, bool &bChangeAll )
+{
+	bChangeAll = false;
+
+	if( Producers.IsValidIndex( Index ) && Producers.Num() > 1 )
+	{
+		FDataprepAssetProducer& InAssetProducer = Producers[Index];
+
+		// Check if input producer is still superseded if applicable
+		if( InAssetProducer.SupersededBy != INDEX_NONE )
+		{
+			FDataprepAssetProducer& SupersedingAssetProducer = Producers[ InAssetProducer.SupersededBy ];
+
+			if( !SupersedingAssetProducer.Producer->Supersede( InAssetProducer.Producer ) )
+			{
+				InAssetProducer.SupersededBy = INDEX_NONE;
+			}
+		}
+
+		// Check if producer is now superseded by any other producer
+		int32 SupersederIndex = 0;
+		for( FDataprepAssetProducer& AssetProducer : Producers )
+		{
+			if( AssetProducer.Producer != InAssetProducer.Producer && AssetProducer.Producer->Supersede( InAssetProducer.Producer ) )
+			{
+				// Disable found producer if another producer supersedes its production
+				InAssetProducer.SupersededBy = SupersederIndex;
+				break;
+			}
+			SupersederIndex++;
+		}
+
+		// If input producer superseded any other producer, check if this is still valid
+		// Check if input producer does not supersede other producers
+		for( FDataprepAssetProducer& AssetProducer : Producers )
+		{
+			if( AssetProducer.Producer != InAssetProducer.Producer )
+			{
+				if( AssetProducer.SupersededBy == Index )
+				{
+					if( !InAssetProducer.Producer->Supersede( AssetProducer.Producer ) )
+					{
+						bChangeAll = true;
+						AssetProducer.SupersededBy = INDEX_NONE;
+					}
+				}
+				else if( InAssetProducer.Producer->Supersede( AssetProducer.Producer ) )
+				{
+					bChangeAll = true;
+					AssetProducer.SupersededBy = Index;
+				}
+			}
+		}
+	}
+}
\ No newline at end of file
diff --git a/Engine/Plugins/Experimental/DataPrepEditor/Source/DataPrepCore/Private/DataPrepContentConsumer.cpp b/Engine/Plugins/Experimental/DataPrepEditor/Source/DataPrepCore/Private/DataPrepContentConsumer.cpp
--- a/Engine/Plugins/Experimental/DataPrepEditor/Source/DataPrepCore/Private/DataPrepContentConsumer.cpp
+++ b/Engine/Plugins/Experimental/DataPrepEditor/Source/DataPrepCore/Private/DataPrepContentConsumer.cpp
@@ -46,12 +46,14 @@ bool UDataprepContentConsumer::SetTargetContentFolder(const FString& InTargetCon
 	if( InTargetContentFolder.IsEmpty() )
 	{
 		TargetContentFolder = FPaths::GetPath( GetOutermost()->GetPathName() );
+		OnChanged.Broadcast();
 		return true;
 	}
 
 	if( FPackageName::IsValidLongPackageName( InTargetContentFolder ) )
 	{
 		TargetContentFolder = InTargetContentFolder;
+		OnChanged.Broadcast();
 		return true;
 	}
 
diff --git a/Engine/Plugins/Experimental/DataPrepEditor/Source/DataPrepCore/Private/DataPrepFactories.cpp b/Engine/Plugins/Experimental/DataPrepEditor/Source/DataPrepCore/Private/DataPrepFactories.cpp
--- a/Engine/Plugins/Experimental/DataPrepEditor/Source/DataPrepCore/Private/DataPrepFactories.cpp
+++ b/Engine/Plugins/Experimental/DataPrepEditor/Source/DataPrepCore/Private/DataPrepFactories.cpp
@@ -23,27 +23,8 @@ UObject * UDataprepAssetFactory::FactoryCreateNew(UClass* InClass, UObject* InPa
 	check( InClass->IsChildOf( UDataprepAsset::StaticClass() ) );
 	UDataprepAsset* DataprepAsset = NewObject<UDataprepAsset>(InParent, InClass, InName, Flags | RF_Transactional);
 
-	// Set DataprepAsset's consumer to the first registered consumer
-	for( TObjectIterator< UClass > It ; It ; ++It )
-	{
-		UClass* CurrentClass = (*It);
-
-		if ( !CurrentClass->HasAnyClassFlags( CLASS_Abstract ) )
-		{
-			if( CurrentClass->IsChildOf( UDataprepContentConsumer::StaticClass() ) )
-			{
-				DataprepAsset->Consumer = NewObject< UDataprepContentConsumer >( DataprepAsset->GetOutermost(), CurrentClass, NAME_None, RF_Transactional );
-				check( DataprepAsset->Consumer );
-
-				FAssetRegistryModule::AssetCreated( DataprepAsset->Consumer );
-				DataprepAsset->Consumer->MarkPackageDirty();
-
-				break;
-			}
-		}
-	}
-
-	// #ueent_todo: Temp code for the nodes development
+	FAssetRegistryModule::AssetCreated( DataprepAsset );
+	DataprepAsset->MarkPackageDirty();
 
 	return DataprepAsset;
 }
diff --git a/Engine/Plugins/Experimental/DataPrepEditor/Source/DataPrepCore/Public/DataPrepAsset.h b/Engine/Plugins/Experimental/DataPrepEditor/Source/DataPrepCore/Public/DataPrepAsset.h
--- a/Engine/Plugins/Experimental/DataPrepEditor/Source/DataPrepCore/Public/DataPrepAsset.h
+++ b/Engine/Plugins/Experimental/DataPrepEditor/Source/DataPrepCore/Public/DataPrepAsset.h
@@ -4,18 +4,23 @@
 
 #include "CoreMinimal.h"
 
+#include "DataPrepContentConsumer.h"
 #include "DataPrepContentProducer.h"
 
 #include "DataPrepAsset.generated.h"
 
 class UDataprepActionAsset;
 class UDataprepContentConsumer;
-class UDataprepContentProducer;
 class UDataprepAsset;
 
-// Delegates
-DECLARE_MULTICAST_DELEGATE(FOnActionsOrderChanged)
-DECLARE_MULTICAST_DELEGATE_OneParam(FOnActionOperationsOrderChanged, const UDataprepActionAsset* /** InActionAsset The modified action */)
+enum class FDataprepAssetChangeType : uint8
+{
+	ProducerAdded,
+	ProducerRemoved,
+	ProducerModified,
+	ConsumerModified,
+	BlueprintModified,
+};
 
 USTRUCT(Experimental)
 struct FDataprepAssetProducer
@@ -25,18 +30,23 @@ struct FDataprepAssetProducer
 	FDataprepAssetProducer()
 		: Producer(nullptr)
 		, bIsEnabled( true )
+		, SupersededBy( INDEX_NONE )
 	{}
 
 	FDataprepAssetProducer(UDataprepContentProducer* InProducer, bool bInEnabled )
 		: Producer(InProducer)
 		, bIsEnabled( bInEnabled )
+		, SupersededBy( INDEX_NONE )
 	{}
 
 	UPROPERTY()
 	UDataprepContentProducer* Producer;
-	
+
 	UPROPERTY()
 	bool bIsEnabled;
+
+	UPROPERTY()
+	int32 SupersededBy;
 };
 
 /**
@@ -53,14 +63,14 @@ struct FDataprepAssetAction
 		: ActionAsset(nullptr)
 		, bIsEnabled(true)
 		, DataprepAssetPtr(nullptr)
-		, OnOperationOrderChandedHandle()
+		, OnOperationOrderChangedHandle()
 	{}
 
 	FDataprepAssetAction(UDataprepActionAsset* InActionAsset, bool bInEnabled, UDataprepAsset& DataprepAsset)
 		: ActionAsset(nullptr)
 		, bIsEnabled(bInEnabled)
 		, DataprepAssetPtr(&DataprepAsset)
-		, OnOperationOrderChandedHandle()
+		, OnOperationOrderChangedHandle()
 	{
 		SetActionAsset(InActionAsset);
 	}
@@ -69,7 +79,7 @@ struct FDataprepAssetAction
 		: ActionAsset(nullptr)
 		, bIsEnabled(Other.bIsEnabled)
 		, DataprepAssetPtr(Other.DataprepAssetPtr)
-		, OnOperationOrderChandedHandle()
+		, OnOperationOrderChangedHandle()
 	{
 		SetActionAsset(Other.ActionAsset);
 	}
@@ -78,7 +88,7 @@ struct FDataprepAssetAction
 		: ActionAsset(nullptr)
 		, bIsEnabled(Other.bIsEnabled)
 		, DataprepAssetPtr(Other.DataprepAssetPtr)
-		, OnOperationOrderChandedHandle()
+		, OnOperationOrderChangedHandle()
 	{
 		SetActionAsset(Other.ActionAsset);
 		Other.SetActionAsset( nullptr );
@@ -116,7 +126,7 @@ private:
 
 	TWeakObjectPtr<UDataprepAsset> DataprepAssetPtr;
 
-	FDelegateHandle OnOperationOrderChandedHandle;
+	FDelegateHandle OnOperationOrderChangedHandle;
 };
 
 
@@ -130,6 +140,11 @@ public:
 
 	virtual ~UDataprepAsset();
 
+	// UObject interface
+	virtual void Serialize( FArchive& Ar ) override;
+	virtual void PostInitProperties() override;
+	// End of UObject interface
+
 	/**
 	 * Create an action at the end of actions array
 	 * @return The index of the added action
@@ -190,17 +205,113 @@ public:
 	 * Allow an observer to be notified when the actions order changed that also include adding and removing an action
 	 * @return The delegate that will be broadcasted when the actions order changed
 	 */
-	FOnActionsOrderChanged& GetOnActionsOrderChanged();
+	DECLARE_EVENT(UDataprepAsset, FOnActionsOrderChanged)
+	FOnActionsOrderChanged& GetOnActionsOrderChanged()
+	{
+		return OnActionsOrderChanged;
+	}
 
 	/**
 	 * Allow an observer to be notified when the operations order of an action changed that also include adding and removing an operation
 	 * @return The delegate that will be broadcasted when the operations order of an action has changed
 	 */
-	FOnActionOperationsOrderChanged& GetOnActionOperationsOrderChanged();
+	DECLARE_EVENT_OneParam(UDataprepAsset, FOnActionOperationsOrderChanged, const UDataprepActionAsset* /** InActionAsset The modified action */)
+	FOnActionOperationsOrderChanged& GetOnActionOperationsOrderChanged()
+	{
+		return OnActionOperationsOrderChanged;
+	}
+
+	/**
+	 * Add a producer of a given class
+	 * @param ProducerClass Class of the Producer to add
+	 * @return true if addition was successful, false otherwise
+	 */
+	bool AddProducer( UClass* ProducerClass );
+
+	/**
+	 * Remove the producer at a given index
+	 * @param IndexToRemove Index of the producer to remove
+	 * @return true if remove was successful, false otherwise
+	*/
+	bool RemoveProducer( int32 IndexToRemove );
+
+	/**
+	 * Enable/Disable the producer at a given index
+	 * @param Index Index of the producer to update
+	 */
+	void EnableProducer(int32 Index, bool bValue);
+
+	/**
+	 * Enable/Disable the producer at a given index
+	 * @param Index Index of the producer to update
+	 */
+	bool EnableAllProducers(bool bValue);
+
+	/**
+	 * Toggle the producer at a given index
+	 * @param Index Index of the producer to update
+	 */
+	void ToggleProducer( int32 Index )
+	{
+		EnableProducer( Index, !IsProducerEnabled( Index ) );
+	}
+
+	/**
+	 * Replace the current consumer with one of a given class
+	 * @param NewConsumerClass Class of the consumer to replace the current one with
+	 * @return true if replacement was successful, false otherwise
+	 */
+	bool ReplaceConsumer( UClass* NewConsumerClass );
+
+	const UDataprepContentConsumer* GetConsumer() const { return Consumer; }
+
+	UDataprepContentConsumer* GetConsumer() { return Consumer; }
+
+	int32 GetProducersCount() const { return Producers.Num(); }
+
+	/** @return pointer on producer at Index-th position in Producers array. nullptr if Index is invalid */
+	const UDataprepContentProducer* GetProducer(int32 Index) const
+	{ 
+		return Producers.IsValidIndex(Index) ? Producers[Index].Producer : nullptr;
+	}
+
+	UDataprepContentProducer* GetProducer(int32 Index)
+	{ 
+		return Producers.IsValidIndex(Index) ? Producers[Index].Producer : nullptr;
+	}
+
+	/** @return pointer on producer at Index-th position in Producers array. nullptr if Index is invalid */
+	bool IsProducerEnabled(int32 Index) const
+	{ 
+		return Producers.IsValidIndex(Index) ? Producers[Index].bIsEnabled : nullptr;
+	}
+
+	/** @return pointer on producer at Index-th position in Producers array. nullptr if Index is invalid */
+	bool IsProducerSuperseded(int32 Index) const
+	{ 
+		return Producers.IsValidIndex(Index) ? Producers[Index].SupersededBy != INDEX_NONE : false;
+	}
+
+	/**
+	 * Allow an observer to be notified when when the consumer or one of the producer has changed
+	 * @return The delegate that will be broadcasted when when the consumer or one of the producer has changed
+	 */
+	DECLARE_EVENT_TwoParams(UDataprepAsset, FOnDataprepAssetChanged, FDataprepAssetChangeType, int32 )
+	FOnDataprepAssetChanged& GetOnChanged() { return OnChanged; }
 
 	void RunProducers( const UDataprepContentProducer::ProducerContext& InContext, TArray< TWeakObjectPtr< UObject > >& OutAssets );
 
+	bool RunConsumer( const UDataprepContentConsumer::ConsumerContext& InContext, FString& OutReason );
+
 #if WITH_EDITORONLY_DATA
+
+	// Temp code for the nodes development
+	/** Temporary: Pointer to data preparation pipeline blueprint used to process input data */
+	UPROPERTY()
+	class UBlueprint* DataprepRecipeBP;
+	// end of temp code for nodes development
+
+protected:
 	/** List of producers referenced by the asset */
 	UPROPERTY()
 	TArray< FDataprepAssetProducer > Producers;
@@ -210,13 +321,28 @@ public:
 	UDataprepContentConsumer* Consumer;
 #endif
 
+protected:
+
+	/**
+	 * Answer notification that the consumer has changed
+	 * @param InProducer Producer which has changed
+	 */
+	void OnProducerChanged( const UDataprepContentProducer* InProducer );
+
+	/** Answer notification that the consumer has changed */
+	void OnConsumerChanged();
+
 	// Temp code for the nodes development
-	/** Temporary: Pointer to data preparation pipeline blueprint used to process input data */
-	UPROPERTY()
-	class UBlueprint* DataprepRecipeBP;
+	void OnBlueprintChanged( class UBlueprint* InBlueprint );
 	// end of temp code for nodes development
 
 private:
+	/**
+	  * Helper to check superseding on all producers after changes made on one producer
+	  * @param Index Index of the producer to validate against
+	  * @param bChangeAll Set to true if changes on input producer implied changes on other producers
+	  */
+	void ValidateProducerChanges( int32 Index, bool &bChangeAll );
 
 	void RemoveInvalidActions();
 
@@ -224,11 +350,14 @@ private:
 	UPROPERTY()
 	TArray<FDataprepAssetAction> Actions;
 
-	//Delegate broadcasted when the actions order changed
+	/** Delegate broadcasted when the actions order changed */
 	FOnActionsOrderChanged OnActionsOrderChanged;
 
-	//Delegate broadcasted when the order of the operations of an action changed
+	/** Delegate broadcasted when the order of the operations of an action changed */
 	FOnActionOperationsOrderChanged OnActionOperationsOrderChanged;
 
+	/** Delegate broadcasted when the consumer or one of the producers has changed */
+	FOnDataprepAssetChanged OnChanged;
+
 	FDelegateHandle OnAssetDeletedHandle;
 };
\ No newline at end of file
diff --git a/Engine/Plugins/Experimental/DataPrepEditor/Source/DataPrepCore/Public/DataPrepContentConsumer.h b/Engine/Plugins/Experimental/DataPrepEditor/Source/DataPrepCore/Public/DataPrepContentConsumer.h
--- a/Engine/Plugins/Experimental/DataPrepEditor/Source/DataPrepCore/Public/DataPrepContentConsumer.h
+++ b/Engine/Plugins/Experimental/DataPrepEditor/Source/DataPrepCore/Public/DataPrepContentConsumer.h
@@ -120,8 +120,8 @@ public:
 	/**
 	 * Sets the path of the package the consumer should move assets to if applicable.
 	 * Generally, this package path is substituted to the temporary path the assets are in
-	 * @param InPackagePath : New package path
-	 * @return true if the package path is valid
+	 * @param InTargetContentFolder : Path of the package to save any assets in
+	 * @return true if the assignment has been successful, false otherwise
 	 * @remark if InPackagePath is empty the package path of the consumer is used
 	 */
 	virtual bool SetTargetContentFolder(const FString& InTargetContentFolder );
@@ -130,6 +130,16 @@ public:
 
 	const FString& GetTargetContentFolder() { return TargetContentFolder; }
 
+	/**
+	 * Allow an observer to be notified when one of the properties of the consumer changes
+	 * @return The delegate that will be broadcasted when the consumer changed
+	 */
+	DECLARE_EVENT(UDataprepContentConsumer, FDataprepConsumerChanged)
+	FDataprepConsumerChanged& GetOnChanged()
+	{
+		return OnChanged;
+	}
+
 protected:
 
 	// Start of helper functions to log messages and report progress
@@ -183,4 +193,7 @@ protected:
 
 	/** Context which the consumer will run with */
 	ConsumerContext Context;
+
+	/** Delegate to broadcast changes to the consumer */
+	FDataprepConsumerChanged OnChanged;
 };
\ No newline at end of file
diff --git a/Engine/Plugins/Experimental/DataPrepEditor/Source/DataPrepCore/Public/DataPrepContentProducer.h b/Engine/Plugins/Experimental/DataPrepEditor/Source/DataPrepCore/Public/DataPrepContentProducer.h
--- a/Engine/Plugins/Experimental/DataPrepEditor/Source/DataPrepCore/Public/DataPrepContentProducer.h
+++ b/Engine/Plugins/Experimental/DataPrepEditor/Source/DataPrepCore/Public/DataPrepContentProducer.h
@@ -108,6 +108,25 @@ public:
 	 */
 	virtual FString GetNamespace() const;
 
+	/**
+	 * Allow an observer to be notified when one of the properties of the producer changes
+	 * @return The delegate that will be broadcasted when the consumer changed
+	 * @remark: Subclass of UDataprepContentProducer must use this mechanism to communicate changes to the owning Dataprep asset
+	 */
+	DECLARE_EVENT_OneParam( UDataprepContentProducer, FDataprepProducerChanged, const UDataprepContentProducer* )
+	FDataprepProducerChanged& GetOnChanged()
+	{
+		return OnChanged;
+	}
+
+	/**
+	 * A producer supersede another if its produces the same content or a super-set of the other's content
+	 * @param OtherProducer : Other producer to compare to.
+	 * @return true if the other producer produces the same content or a sub-set of it
+	 * @remark Each sub-class must implement this method
+	 */
+	virtual bool Supersede(const UDataprepContentProducer* OtherProducer) const { unimplemented(); return true; }
+
 protected:
 
 	/**
@@ -166,4 +185,7 @@ protected:
 
 	/** Array of assets generated after a call to Run */
 	TArray<TWeakObjectPtr<UObject>> Assets;
+
+	/** Delegate to broadcast changes to the producer */
+	FDataprepProducerChanged OnChanged;
 };
\ No newline at end of file
diff --git a/Engine/Plugins/Experimental/DataPrepEditor/Source/DataPrepEditor/Private/BlueprintNodes/K2Node_DataprepConsumer.cpp b/Engine/Plugins/Experimental/DataPrepEditor/Source/DataPrepEditor/Private/BlueprintNodes/K2Node_DataprepConsumer.cpp
deleted file mode 100644
diff --git a/Engine/Plugins/Experimental/DataPrepEditor/Source/DataPrepEditor/Private/BlueprintNodes/K2Node_DataprepConsumer.h b/Engine/Plugins/Experimental/DataPrepEditor/Source/DataPrepEditor/Private/BlueprintNodes/K2Node_DataprepConsumer.h
deleted file mode 100644
diff --git a/Engine/Plugins/Experimental/DataPrepEditor/Source/DataPrepEditor/Private/BlueprintNodes/K2Node_DataprepProducer.cpp b/Engine/Plugins/Experimental/DataPrepEditor/Source/DataPrepEditor/Private/BlueprintNodes/K2Node_DataprepProducer.cpp
--- a/Engine/Plugins/Experimental/DataPrepEditor/Source/DataPrepEditor/Private/BlueprintNodes/K2Node_DataprepProducer.cpp
+++ b/Engine/Plugins/Experimental/DataPrepEditor/Source/DataPrepEditor/Private/BlueprintNodes/K2Node_DataprepProducer.cpp
@@ -10,18 +10,19 @@
 #include "KismetCompiler.h"
 #include "ScopedTransaction.h"
 
+#include "EditorFontGlyphs.h"
 #include "EditorStyleSet.h"
 #include "Framework/MultiBox/MultiBoxBuilder.h"
-#include "PropertyEditorModule.h"
 #include "HAL/PlatformProcess.h"
 #include "Modules/ModuleManager.h"
+#include "PropertyEditorModule.h"
+#include "SGraphNode.h"
 #include "Widgets/Images/SImage.h"
 #include "Widgets/Input/SButton.h"
 #include "Widgets/Input/SComboButton.h"
+#include "Widgets/Layout/SBox.h"
 #include "Widgets/Layout/SGridPanel.h"
 #include "Widgets/Text/SInlineEditableTextBlock.h"
-#include "SGraphNode.h"
-#include "Widgets/Layout/SBox.h"
 
 #define LOCTEXT_NAMESPACE "UK2Node_DataprepProducer"
 
@@ -102,10 +103,10 @@ TSharedPtr<SGraphNode> UK2Node_DataprepProducer::CreateVisualWidget()
 	{
 		struct ProducerWidget
 		{
-			FDataprepAssetProducer* ProducerPtr;
+			int32 ProducerIndex;
 			TSharedPtr< STextBlock > CheckBox;
 
-			ProducerWidget() : ProducerPtr(nullptr) {}
+			ProducerWidget() : ProducerIndex( INDEX_NONE ) {}
 		};
 
 	public:
@@ -114,7 +115,10 @@ TSharedPtr<SGraphNode> UK2Node_DataprepProducer::CreateVisualWidget()
 
 		~SGraphNodeDataprepProducer()
 		{
-			FCoreUObjectDelegates::OnObjectPropertyChanged.RemoveAll(this);
+			if( DataprepAsset )
+			{
+				DataprepAsset->GetOnChanged().RemoveAll( this );
+			}
 		}
 
 		void Construct(const FArguments& InArgs, UK2Node_DataprepProducer* InNode)
@@ -134,10 +138,11 @@ TSharedPtr<SGraphNode> UK2Node_DataprepProducer::CreateVisualWidget()
 			SGraphNode::CreateStandardPinWidget(Pin);
 		}
 
-		void OnPostPropertyChanged( UObject* InObject, FPropertyChangedEvent& InPropertyChangedEvent)
+		void OnDataprepAssetChanged( FDataprepAssetChangeType ChangeType, int32 Index)
 		{
-			const FName MemberPropertyName = InPropertyChangedEvent.MemberProperty != nullptr ? InPropertyChangedEvent.MemberProperty->GetFName() : NAME_None;
-			if (InObject && MemberPropertyName == GET_MEMBER_NAME_CHECKED(UDataprepAsset, Producers))
+			if( ChangeType == FDataprepAssetChangeType::ProducerModified ||
+				ChangeType == FDataprepAssetChangeType::ProducerAdded ||
+				ChangeType == FDataprepAssetChangeType::ProducerRemoved )
 			{
 				UpdateGraphNode();
 			}
@@ -148,9 +153,8 @@ TSharedPtr<SGraphNode> UK2Node_DataprepProducer::CreateVisualWidget()
 			bNeedsUpdate = false;
 
 			DataprepAsset = CastChecked<UK2Node_DataprepProducer>(GraphNode)->GetDataprepAsset();
-			check( DataprepAsset );
 
-			FCoreUObjectDelegates::OnObjectPropertyChanged.AddRaw(this, &SGraphNodeDataprepProducer::OnPostPropertyChanged);
+			DataprepAsset->GetOnChanged().AddRaw( this, &SGraphNodeDataprepProducer::OnDataprepAssetChanged );
 
 			// #ueent_todo: Move this code where it would capture new classes of producer created at runtime, .i.e BP based producers
 			for( TObjectIterator< UClass > It ; It ; ++It )
@@ -175,17 +179,7 @@ TSharedPtr<SGraphNode> UK2Node_DataprepProducer::CreateVisualWidget()
 			UClass* ProducerClass = ProducerDescriptions[Index].Get<0>();
 			check( ProducerClass );
 
-			UDataprepContentProducer* Producer = NewObject< UDataprepContentProducer >( DataprepAsset->GetOutermost(), ProducerClass, NAME_None, RF_Transactional );
-			FAssetRegistryModule::AssetCreated( Producer );
-			Producer->MarkPackageDirty();
-
-			FScopedTransaction Transaction(LOCTEXT("DataprepAsset_AddProducer", "AddProducer"));
-
-			DataprepAsset->Producers.Emplace( Producer, true );
-
-			UProperty* Property = FindField<UProperty>(UDataprepAsset::StaticClass(), GET_MEMBER_NAME_CHECKED(UDataprepAsset, Producers));
-			FPropertyChangedEvent PropertyChangedEvent( Property, EPropertyChangeType::ArrayAdd );
-			DataprepAsset->PostEditChangeProperty( PropertyChangedEvent );
+			DataprepAsset->AddProducer( ProducerClass );
 		}
 
 		TSharedRef<SWidget> CreateAddProducerMenuWidget()
@@ -218,51 +212,22 @@ TSharedPtr<SGraphNode> UK2Node_DataprepProducer::CreateVisualWidget()
 			return MenuBuilder.MakeWidget();
 		}
 
-		void DeleteRegisteredAsset(UObject* Asset)
-		{
-			if(Asset != nullptr)
-			{
-				FName AssetName = MakeUniqueObjectName( GetTransientPackage(), Asset->GetClass(), *( Asset->GetName() + TEXT( "_Deleted" ) ) );
-				Asset->Rename( *AssetName.ToString(), GetTransientPackage(), REN_DontCreateRedirectors | REN_NonTransactional );
-
-				Asset->ClearFlags(RF_Standalone | RF_Public);
-				Asset->RemoveFromRoot();
-				Asset->MarkPendingKill();
-
-				FAssetRegistryModule::AssetDeleted( Asset ) ;
-			}
-		}
-
 		TSharedRef< SWidget > CreateProducerWidget( ProducerWidget* ProducerWidgetPtr, int32 Index )
 		{
-			auto CheckEntry = [this, ProducerWidgetPtr]()
+			auto CheckEntry = [this, Index]()
 			{
-				ProducerWidgetPtr->ProducerPtr->bIsEnabled = !ProducerWidgetPtr->ProducerPtr->bIsEnabled;
-				ProducerWidgetPtr->CheckBox->SetText( FText::FromString( ProducerWidgetPtr->ProducerPtr->bIsEnabled ? TEXT( "\xf14a" ) : TEXT( "\xf0c8" ) ) );
-
-				UProperty* Property = FindField<UProperty>(UDataprepAsset::StaticClass(), GET_MEMBER_NAME_CHECKED(UDataprepAsset, Producers));
-				FPropertyChangedEvent PropertyChangedEvent( Property, EPropertyChangeType::ValueSet );
-				DataprepAsset->PostEditChangeProperty( PropertyChangedEvent );
-
+				FScopedTransaction Transaction(LOCTEXT("DataprepAsset_AddProducer", "AddProducer"));
+				DataprepAsset->ToggleProducer( Index );
 				return FReply::Handled();
 			};
 
-			auto DeleteEntry = [this, ProducerWidgetPtr, Index]()
+			auto DeleteEntry = [this, Index]()
 			{
-				FScopedTransaction Transaction(LOCTEXT("DataprepAsset_AddProducer", "AddProducer"));
-
-				DeleteRegisteredAsset( ProducerWidgetPtr->ProducerPtr->Producer );
-				DataprepAsset->Producers.RemoveAt( Index );
-
-				UProperty* Property = FindField<UProperty>(UDataprepAsset::StaticClass(), GET_MEMBER_NAME_CHECKED(UDataprepAsset, Producers));
-				FPropertyChangedEvent PropertyChangedEvent( Property, EPropertyChangeType::ArrayRemove );
-				DataprepAsset->PostEditChangeProperty( PropertyChangedEvent );
-
+				FScopedTransaction Transaction(LOCTEXT("DataprepAsset_RemoveProducer", "RemoveProducer"));
+				DataprepAsset->RemoveProducer( Index );
 				return FReply::Handled();
 			};
 
-			FPropertyEditorModule& PropertyEditorModule = FModuleManager::GetModuleChecked<FPropertyEditorModule>("PropertyEditor");
-
 			FDetailsViewArgs DetailsViewArgs;
 			DetailsViewArgs.bAllowSearch = false;
 			DetailsViewArgs.bAllowMultipleTopLevelObjects = true;
@@ -291,7 +256,7 @@ TSharedPtr<SGraphNode> UK2Node_DataprepProducer::CreateVisualWidget()
 						SAssignNew( ProducerWidgetPtr->CheckBox, STextBlock )
 						.Font( FEditorStyle::Get().GetFontStyle( DataprepProducerFontName ) )
 						.ColorAndOpacity( FLinearColor::White )
-						.Text(FText::FromString( FString( ProducerWidgetPtr->ProducerPtr->bIsEnabled ? TEXT( "\xf14a" ) : TEXT( "\xf0c8" ) ) ) )
+						.Text( DataprepAsset->IsProducerEnabled( Index ) ? FEditorFontGlyphs::Check_Square : FEditorFontGlyphs::Square )
 					]
 				]
 				// Input entry label
@@ -316,16 +281,18 @@ TSharedPtr<SGraphNode> UK2Node_DataprepProducer::CreateVisualWidget()
 						SNew( STextBlock )
 						.Font( FEditorStyle::Get().GetFontStyle( DataprepProducerFontName ) )
 						.ColorAndOpacity( FLinearColor::White )
-						.Text( FText::FromString( FString( TEXT("\xf1f8") ) ) )
+						.Text( FEditorFontGlyphs::Trash )
 					]
 				]
 			];
 
-			TSharedPtr<IDetailsView> DetailsView = PropertyEditorModule.CreateDetailView(DetailsViewArgs);
+			FPropertyEditorModule& PropertyEditorModule = FModuleManager::GetModuleChecked<FPropertyEditorModule>("PropertyEditor");
+
+			TSharedPtr<IDetailsView> DetailsView = PropertyEditorModule.CreateDetailView( DetailsViewArgs );
 
 			DetailsViewBox->SetContent(DetailsView.ToSharedRef());
 
-			DetailsView->SetObject( ProducerWidgetPtr->ProducerPtr->Producer );
+			DetailsView->SetObject( DataprepAsset->GetProducer( Index ) );
 
 			return Widget.ToSharedRef();
 		}
@@ -376,12 +343,12 @@ TSharedPtr<SGraphNode> UK2Node_DataprepProducer::CreateVisualWidget()
 				]
 			];
 
-			ProducerWidgets.SetNum( DataprepAsset->Producers.Num() );
+			const int32 ProducersCount = DataprepAsset->GetProducersCount();
+			ProducerWidgets.SetNum( ProducersCount );
 
-			int32 Index = 0;
-			for( FDataprepAssetProducer& Producer : DataprepAsset->Producers )
+			for( int32 Index = 0; Index < ProducersCount; ++Index )
 			{
-				ProducerWidgets[ Index ].ProducerPtr = &Producer;
+				ProducerWidgets[ Index ].ProducerIndex = Index;
 
 				MainBox->AddSlot()[
 					CreateProducerWidget( &ProducerWidgets[ Index ], Index )
diff --git a/Engine/Plugins/Experimental/DataPrepEditor/Source/DataPrepEditor/Private/DataPrepEditor.cpp b/Engine/Plugins/Experimental/DataPrepEditor/Source/DataPrepEditor/Private/DataPrepEditor.cpp
--- a/Engine/Plugins/Experimental/DataPrepEditor/Source/DataPrepEditor/Private/DataPrepEditor.cpp
+++ b/Engine/Plugins/Experimental/DataPrepEditor/Source/DataPrepEditor/Private/DataPrepEditor.cpp
@@ -3,7 +3,6 @@
 #include "DataPrepEditor.h"
 
 #include "BlueprintNodes/K2Node_DataprepAction.h"
-#include "BlueprintNodes/K2Node_DataprepConsumer.h"
 #include "BlueprintNodes/K2Node_DataprepProducer.h"
 #include "DataPrepContentConsumer.h"
 #include "DataPrepContentProducer.h"
@@ -18,6 +17,7 @@
 #include "Widgets/SDataprepPalette.h"
 
 #include "ActorEditorUtils.h"
+#include "AssetDeleteModel.h"
 #include "AssetRegistryModule.h"
 #include "BlueprintNodeSpawner.h"
 #include "DesktopPlatformModule.h"
@@ -29,6 +29,7 @@
 #include "EngineUtils.h"
 #include "Framework/Application/SlateApplication.h"
 #include "Framework/MultiBox/MultiBoxBuilder.h"
+#include "GenericPlatform/GenericPlatformTime.h"
 #include "HAL/FileManager.h"
 #include "Kismet2/BlueprintEditorUtils.h"
 #include "Kismet2/KismetEditorUtilities.h"
@@ -183,6 +184,11 @@ FDataprepEditor::FDataprepEditor()
 
 FDataprepEditor::~FDataprepEditor()
 {
+	if( DataprepAssetPtr.IsValid() )
+	{
+		DataprepAssetPtr->GetOnChanged().RemoveAll( this );
+	}
+
 	if ( PreviewWorld )
 	{
 		GEngine->DestroyWorldContext( PreviewWorld.Get() );
@@ -271,7 +277,7 @@ void FDataprepEditor::InitDataprepEditor(const EToolkitMode::Type Mode, const TS
 	DataprepAssetPtr = TWeakObjectPtr<UDataprepAsset>(InDataprepAsset);
 	check( DataprepAssetPtr.IsValid() );
 
-	FCoreUObjectDelegates::OnObjectPropertyChanged.AddRaw(this, &FDataprepEditor::OnDataprepAssetChanged);
+	DataprepAssetPtr->GetOnChanged().AddRaw( this, &FDataprepEditor::OnDataprepAssetChanged );
 
 	// Assign unique session identifier
 	SessionID = FGuid::NewGuid().ToString();
@@ -282,75 +288,39 @@ void FDataprepEditor::InitDataprepEditor(const EToolkitMode::Type Mode, const TS
 
 	// Temp code for the nodes development
 	DataprepRecipeBPPtr = DataprepAssetPtr->DataprepRecipeBP;
-	if( !DataprepRecipeBPPtr.IsValid() )
-	{
-		const FString DesiredName = DataprepAssetPtr->GetName() + TEXT("_Recipe");
-		FName BlueprintName = MakeUniqueObjectName( DataprepAssetPtr->GetOutermost(), UBlueprint::StaticClass(), *DesiredName );
+	check( DataprepRecipeBPPtr.IsValid() );
 
-		UBlueprint* DataprepRecipeBP = FKismetEditorUtilities::CreateBlueprint( UDataprepRecipe::StaticClass(), DataprepAssetPtr.Get(), BlueprintName
-			, BPTYPE_Normal, UBlueprint::StaticClass(), UBlueprintGeneratedClass::StaticClass() );
+	UEdGraph* PipelineGraph = FBlueprintEditorUtils::FindEventGraph(DataprepRecipeBPPtr.Get());
+	check( PipelineGraph );
 
-		{
-			UEdGraph* EventGraph = FBlueprintEditorUtils::FindEventGraph(DataprepRecipeBP);
-
-			IBlueprintNodeBinder::FBindingSet Bindings;
-			UK2Node_DataprepProducer* ProducerNode = Cast<UK2Node_DataprepProducer>(UBlueprintNodeSpawner::Create<UK2Node_DataprepProducer>()->Invoke(EventGraph, Bindings, FVector2D(-100,0)));
-
-			ProducerNode->SetDataprepAsset( DataprepAssetPtr.Get() );
-			StartNode = ProducerNode;
-		}
+	TArray< class UK2Node_DataprepAction* > ActionNodes;
 
-		FAssetRegistryModule::AssetCreated( DataprepRecipeBP );
-
-		DataprepRecipeBP->MarkPackageDirty();
-
-		DataprepAssetPtr->DataprepRecipeBP = DataprepRecipeBP;
-
-		DataprepRecipeBPPtr = DataprepAssetPtr->DataprepRecipeBP;
-	}
-	else
+	for( UEdGraphNode* GraphNode : PipelineGraph->Nodes )
 	{
-		UEdGraph* PipelineGraph = FBlueprintEditorUtils::FindEventGraph(DataprepRecipeBPPtr.Get());
-		check( PipelineGraph );
-
-		TArray< class UK2Node_DataprepAction* > ActionNodes;
-
-		for( UEdGraphNode* GraphNode : PipelineGraph->Nodes )
+		if( GraphNode->IsA<UK2Node_DataprepProducer>() )
 		{
-			if( GraphNode->IsA<UK2Node_DataprepProducer>() )
-			{
-				StartNode = GraphNode;
-			}
-			else if(StartNode != nullptr)
-			{
-				break;
-			}
+			StartNode = GraphNode;
+		}
+		else if(StartNode != nullptr)
+		{
+			break;
 		}
-		check( StartNode );
 	}
-	// end of temp code for nodes development
 
 	FDataprepEditorStyle::Initialize();
 
-	// Find content producers and consumers the user could use for his/her data preparation
-	for( TObjectIterator< UClass > It ; It ; ++It )
+	// This should normally happen only with a brand new Dataprep asset
+	if( StartNode == nullptr )
 	{
-		UClass* CurrentClass = (*It);
+		UEdGraph* EventGraph = FBlueprintEditorUtils::FindEventGraph(DataprepRecipeBPPtr.Get());
 
-		if ( !CurrentClass->HasAnyClassFlags( CLASS_Abstract ) )
-		{
-			if( CurrentClass->IsChildOf( UDataprepContentProducer::StaticClass() ) )
-			{
-				UDataprepContentProducer* Producer = Cast< UDataprepContentProducer >( CurrentClass->GetDefaultObject() );
-				ProducerDescriptions.Emplace( CurrentClass, Producer->GetLabel(), Producer->GetDescription() );
-			}
-			else if( CurrentClass->IsChildOf( UDataprepContentConsumer::StaticClass() ) )
-			{
-				UDataprepContentConsumer* Consumer = Cast< UDataprepContentConsumer >( CurrentClass->GetDefaultObject() );
-				ConsumerDescriptions.Emplace( CurrentClass, Consumer->GetLabel(), Consumer->GetDescription() );
-			}
-		}
+		IBlueprintNodeBinder::FBindingSet Bindings;
+		UK2Node_DataprepProducer* ProducerNode = Cast<UK2Node_DataprepProducer>(UBlueprintNodeSpawner::Create<UK2Node_DataprepProducer>()->Invoke(EventGraph, Bindings, FVector2D(-100,0)));
+
+		ProducerNode->SetDataprepAsset( DataprepAssetPtr.Get() );
+		StartNode = ProducerNode;
 	}
+	// end of temp code for nodes development
 
 	GEditor->RegisterForUndo(this);
 
@@ -510,12 +480,15 @@ void FDataprepEditor::OnBuildWorld()
 	DataprepRecipePtr->ResetAssets();
 	// end of temp code for nodes development
 
-	if (DataprepAsset->Producers.Num() == 0)
+	if (DataprepAsset->GetProducersCount() == 0)
 	{
 		ResetBuildWorld();
 		return;
 	}
 
+	uint64 StartTime = FPlatformTime::Cycles64();
+	UE_LOG( LogDataprepEditor, Log, TEXT("Importing ...") );
+
 	CleanPreviewWorld();
 
 	if(UEdGraph* PipelineGraph = GetPipelineGraph())
@@ -529,7 +502,7 @@ void FDataprepEditor::OnBuildWorld()
 	UPackage* TransientPackage = NewObject< UPackage >( nullptr, *GetTransientContentFolder(), RF_Transient );
 	TransientPackage->FullyLoad();
 
-	// #ueent_todo: Add progress reporter and logger to dataprep editor
+	// #ueent_todo: Add progress reporter and logger to Dataprep editor
 	UDataprepContentProducer::ProducerContext Context;
 	Context.SetWorld( PreviewWorld.Get() )
 		.SetRootPackage( TransientPackage )
@@ -541,147 +514,57 @@ void FDataprepEditor::OnBuildWorld()
 	CachedAssets.Reset();
 	CachedAssets.Append( Assets );
 
-	TakeSnapshot();
-	// #ueent_todo: For testing purpose, snapshot is taken then restored
-	//RestoreFromSnapshot();
-
-	OnWorldChanged();
+	UpdatePreviewPanels();
 	bWorldBuilt = true;
 	bIsFirstRun = true;
 	bPipelineExecuted = false;
-}
-
-void FDataprepEditor::OnDataprepAssetChanged(UObject* InObject, FPropertyChangedEvent& InPropertyChangedEvent)
-{
-	if(InObject == nullptr)
-	{
-		return;
-	}
-
-	if(Cast<UDataprepAsset>(InObject) == DataprepAssetPtr.Get())
-	{
-		const FName MemberPropertyName = InPropertyChangedEvent.MemberProperty != nullptr ? InPropertyChangedEvent.MemberProperty->GetFName() : NAME_None;
-
-		if (MemberPropertyName == GET_MEMBER_NAME_CHECKED(UDataprepAsset, Producers))
-		{
-			ResetBuildWorld();
-		}
-		else if (MemberPropertyName == GET_MEMBER_NAME_CHECKED(UDataprepAsset, Consumer))
-		{
-		}
-	}
-
-	if (Cast<UBlueprint>(InObject) == DataprepRecipeBPPtr.Get())
-	{
-	}
-	else if ( InObject->StaticClass()->IsChildOf( UDataprepContentConsumer::StaticClass() ) )
-	{
-	}
-	else if ( InObject->StaticClass()->IsChildOf( UDataprepContentProducer::StaticClass() ) )
-	{
-		ResetBuildWorld();
-	}
-}
-
-void FDataprepEditor::OnRemoveProducer( int32 Index )
-{
-	if (!DataprepAssetPtr.IsValid() || !DataprepAssetPtr->Producers.IsValidIndex( Index ) )
-	{
-		return;
-	}
 
-	FDataprepAssetProducer& AssetProducer = DataprepAssetPtr->Producers[Index];
+	// Log time spent to import incoming file in minutes and seconds
+	double ElapsedSeconds = FPlatformTime::ToSeconds64(FPlatformTime::Cycles64() - StartTime);
 
-	DeleteRegisteredAsset( AssetProducer.Producer );
+	int ElapsedMin = int(ElapsedSeconds / 60.0);
+	ElapsedSeconds -= 60.0 * (double)ElapsedMin;
+	UE_LOG( LogDataprepEditor, Log, TEXT("Import took [%d min %.3f s]"), ElapsedMin, ElapsedSeconds );
 
-	DataprepAssetPtr->Producers.RemoveAt( Index );
-
-	DataprepAssetPtr->MarkPackageDirty();
-
-	DataprepAssetProducerChangedDelegate.Broadcast();
-	ResetBuildWorld();
+	TakeSnapshot();
 }
 
-bool FDataprepEditor::OnChangeConsumer( TSharedPtr<FString>& NewConsumer )
+void FDataprepEditor::OnDataprepAssetChanged( FDataprepAssetChangeType ChangeType, int32 Index )
 {
-	if ( NewConsumer.IsValid() )
+	switch(ChangeType)
 	{
-		UClass* NewConsumerClass = nullptr;
-
-		for ( int32 Index = 0; Index < ConsumerDescriptions.Num(); ++Index)
-		{
-			if (ConsumerDescriptions[Index].Get<1>().ToString() == *NewConsumer.Get() )
+		case FDataprepAssetChangeType::ConsumerModified:
 			{
-				NewConsumerClass = ConsumerDescriptions[Index].Get<0>();
-				break;
+				UpdatePreviewPanels();
 			}
-		}
+			break;
 
-		if (NewConsumerClass != nullptr)
-		{
-			DeleteRegisteredAsset( DataprepAssetPtr->Consumer );
-
-			DataprepAssetPtr->Consumer = NewObject< UDataprepContentConsumer >( DataprepAssetPtr->GetOutermost(), NewConsumerClass, NAME_None, RF_Transactional );
-			check( DataprepAssetPtr->Consumer );
-
-			FAssetRegistryModule::AssetCreated( DataprepAssetPtr->Consumer );
-			DataprepAssetPtr->Consumer->MarkPackageDirty();
-
-			DataprepAssetPtr->MarkPackageDirty();
+		case FDataprepAssetChangeType::BlueprintModified:
+			{
+				// #ueent_todo: Anything to do there ?
+			}
+			break;
 
-			return true;
-		}
-	}
+		case FDataprepAssetChangeType::ProducerAdded:
+		case FDataprepAssetChangeType::ProducerRemoved:
+		case FDataprepAssetChangeType::ProducerModified:
+			{
+				// Just reset world for the time being
+				ResetBuildWorld();
+			}
+			break;
 
-	return false;
-}
+		default:
+			break;
 
-void FDataprepEditor::OnAddProducer( int32 Index )
-{
-	if (!DataprepAssetPtr.IsValid() || !ProducerDescriptions.IsValidIndex( Index ) )
-	{
-		return;
 	}
-
-	UClass* ProducerClass = ProducerDescriptions[Index].Get<0>();
-	check( ProducerClass );
-
-	UDataprepContentProducer* Producer = NewObject< UDataprepContentProducer >( DataprepAssetPtr->GetOutermost(), ProducerClass, NAME_None, RF_Transactional );
-	FAssetRegistryModule::AssetCreated( Producer );
-	Producer->MarkPackageDirty();
-
-	DataprepAssetPtr->Producers.Emplace( Producer, true );
-	DataprepAssetPtr->MarkPackageDirty();
-
-	DataprepAssetProducerChangedDelegate.Broadcast();
-	ResetBuildWorld();
 }
 
-void FDataprepEditor::OnProducerChanged(int32 Index)
-{
-	if (!DataprepAssetPtr.IsValid() )
-	{
-		return;
-	}
-
-	if (Index == INDEX_NONE)
-	{
-		bWorldBuilt = false;
-	}
-	else if( ProducerDescriptions.IsValidIndex( Index ) )
-	{
-		bWorldBuilt = false;
-	}
-
-	DataprepAssetPtr->MarkPackageDirty();
-}
-
-
 void FDataprepEditor::ResetBuildWorld()
 {
 	bWorldBuilt = false;
 	CleanPreviewWorld();
-	OnWorldChanged();
+	UpdatePreviewPanels();
 	DataprepEditorUtil::DeleteTemporaryPackage( GetTransientContentFolder() );
 }
 
@@ -690,36 +573,41 @@ void FDataprepEditor::CleanPreviewWorld()
 	// Destroy all actors in preview world
 	for (ULevel* Level : PreviewWorld->GetLevels())
 	{
-		for (int32 i = 0; i < Level->Actors.Num(); i++)
+		TArray<AActor*> LevelActors( Level->Actors );
+
+		for( AActor* Actor : LevelActors )
 		{
-			AActor* Actor = Level->Actors[i];
-			if (Actor && !DefaultActorsInPreviewWorld.Contains(Actor))
+			if (Actor && !Actor->IsPendingKill() && !DefaultActorsInPreviewWorld.Contains(Actor))
 			{
 				PreviewWorld->EditorDestroyActor(Actor, true);
 				// Since deletion can be delayed, rename to avoid future name collision 
-				Actor->Rename();
+				Actor->Rename(nullptr, GetTransientPackage(), REN_DontCreateRedirectors | REN_ForceNoResetLoaders);
 			}
 		}
 	}
 
-	if(Assets.Num() > 0)
+	const FString TransientContentFolder( GetTransientContentFolder() );
+	CachedAssets.Append( Assets );
+	for( TWeakObjectPtr<UObject>& Asset : CachedAssets )
 	{
-		TArray<UObject*> ObjectsToDelete;
-		ObjectsToDelete.Reserve(Assets.Num());
-
-		for( TWeakObjectPtr<UObject>& Asset : Assets )
+		if( UObject* ObjectToDelete = Asset.Get() )
 		{
-			if( UObject* AssetObject = Asset.Get() )
+			FString PackagePath = ObjectToDelete->GetOutermost()->GetName();
+			if( PackagePath.StartsWith( TransientContentFolder ) )
 			{
-				AssetObject->Rename( nullptr, GetTransientPackage(), REN_DontCreateRedirectors | REN_NonTransactional );
-				ObjectsToDelete.Add( AssetObject );
-			}
-		}
+				FName ObjectName = MakeUniqueObjectName( GetTransientPackage(), ObjectToDelete->GetClass(), *( ObjectToDelete->GetName() + TEXT( "_Deleted" ) ) );
+				ObjectToDelete->Rename( *ObjectName.ToString(), GetTransientPackage(), REN_DontCreateRedirectors | REN_NonTransactional );
 
-		if(ObjectsToDelete.Num() > 0)
-		{
-			int32 Diff = ObjectTools::DeleteObjects(ObjectsToDelete, false) - ObjectsToDelete.Num();
-			ensure( Diff == 0 );
+				// Remove public and standalone flag so garbage collection can delete the object.
+				ObjectToDelete->ClearFlags(RF_Standalone | RF_Public);
+
+				// Mark its package as dirty as we're going to delete it.
+				ObjectToDelete->MarkPendingKill();
+				ObjectToDelete->MarkPackageDirty();
+
+				// Notify the asset registry
+				FAssetRegistryModule::AssetDeleted( ObjectToDelete );
+			}
 		}
 	}
 
@@ -735,7 +623,7 @@ void FDataprepEditor::CleanPreviewWorld()
 
 void FDataprepEditor::OnExecutePipeline()
 {
-	if( DataprepAssetPtr->Consumer == nullptr )
+	if( DataprepAssetPtr->GetConsumer() == nullptr )
 	{
 		return;
 	}
@@ -836,7 +724,7 @@ void FDataprepEditor::OnExecutePipeline()
 					}
 
 					// World may have changed, update asset preview and scene outliner
-					OnWorldChanged();
+					UpdatePreviewPanels();
 					FSlateApplication::Get().PumpMessages();
 					FSlateApplication::Get().Tick();
 				}
@@ -880,15 +768,9 @@ void FDataprepEditor::OnCommitWorld()
 		.SetProgressReporter( TSharedPtr< IDataprepProgressReporter >( new FDataprepProgressReporter( LOCTEXT("Dataprep_CommitWorld", "Committing ...") ) ) );
 
 	FString OutReason;
-	if( DataprepAssetPtr->Consumer->Initialize( Context, OutReason ) )
+	if( !DataprepAssetPtr->RunConsumer( Context, OutReason ) )
 	{
-		// #ueent_todo: Update state of entry: finalizing
-		if ( !DataprepAssetPtr->Consumer->Run() )
-		{
-			// #ueent_todo: Inform execution has failed
-		}
-
-		DataprepAssetPtr->Consumer->Reset();
+		// #ueent_todo: Inform consumer failed
 	}
 
 	PipelineView->ClearSelectionSet();
@@ -1051,12 +933,17 @@ TSharedRef<SDockTab> FDataprepEditor::SpawnTabPalette(const FSpawnTabArgs & Args
 		];
 }
 
-void FDataprepEditor::OnWorldChanged()
+void FDataprepEditor::UpdatePreviewPanels()
 {
 	// #ueent_todo: There should be a event triggered to inform listeners
 	//				   that new assets have been generated.
 	AssetPreviewView->ClearAssetList();
-	AssetPreviewView->SetAssetsList( Assets, FString() );
+	FString SubstitutePath = DataprepAssetPtr->GetOutermost()->GetName();
+	if(!DataprepAssetPtr->GetConsumer()->GetTargetContentFolder().IsEmpty())
+	{
+		SubstitutePath = DataprepAssetPtr->GetConsumer()->GetTargetContentFolder();
+	}
+	AssetPreviewView->SetAssetsList( Assets, GetTransientContentFolder(), SubstitutePath );
 	SceneOutliner->Refresh();
 }
 
@@ -1067,7 +954,7 @@ bool FDataprepEditor::OnRequestClose()
 
 bool FDataprepEditor::CanBuildWorld()
 {
-	return DataprepAssetPtr->Producers.Num() > 0;
+	return DataprepAssetPtr->GetProducersCount() > 0;
 }
 
 bool FDataprepEditor::CanExecutePipeline()
@@ -1078,27 +965,13 @@ bool FDataprepEditor::CanExecutePipeline()
 bool FDataprepEditor::CanCommitWorld()
 {
 	// Execution of pipeline is not required. User can directly commit result of import
-	return bWorldBuilt && DataprepAssetPtr->Consumer != nullptr;
+	return bWorldBuilt && DataprepAssetPtr->GetConsumer() != nullptr;
 }
 
-void FDataprepEditor::DeleteRegisteredAsset(UObject* Asset)
-{
-	if(Asset != nullptr)
-	{
-		FName AssetName = MakeUniqueObjectName( GetTransientPackage(), Asset->GetClass(), *( Asset->GetName() + TEXT( "_Deleted" ) ) );
-		Asset->Rename( *AssetName.ToString(), GetTransientPackage(), REN_DontCreateRedirectors | REN_NonTransactional );
-
-		Asset->ClearFlags(RF_Standalone | RF_Public);
-		Asset->RemoveFromRoot();
-		Asset->MarkPendingKill();
-
-		FAssetRegistryModule::AssetDeleted( Asset ) ;
-	}
-}
 
 FString FDataprepEditor::GetTransientContentFolder()
 {
-	return FPaths::Combine( FPackageName::GetLongPackagePath( GetDataprepAsset()->GetOutermost()->GetName() ), TEXT("Temp") );
+	return FPaths::Combine( TEXT("/DataPrepEditor"), SessionID );
 }
 
 void DataprepEditorUtil::DeleteTemporaryPackage( const FString& PathToDelete )
diff --git a/Engine/Plugins/Experimental/DataPrepEditor/Source/DataPrepEditor/Private/DataPrepEditor.h b/Engine/Plugins/Experimental/DataPrepEditor/Source/DataPrepEditor/Private/DataPrepEditor.h
--- a/Engine/Plugins/Experimental/DataPrepEditor/Source/DataPrepEditor/Private/DataPrepEditor.h
+++ b/Engine/Plugins/Experimental/DataPrepEditor/Source/DataPrepEditor/Private/DataPrepEditor.h
@@ -90,30 +90,9 @@ public:
 		return DataprepAssetPtr.IsValid() ? DataprepAssetPtr.Get() : nullptr;
 	}
 
-	const TArray< DataprepEditorClassDescription >& GetProducerDescriptions()
-	{
-		return ProducerDescriptions;
-	}
-
-	const TArray< DataprepEditorClassDescription >& GetConsumerDescriptions()
-	{
-		return ConsumerDescriptions;
-	}
-
 	/** Gets or sets the flag for context sensitivity in the graph action menu */
 	bool& GetIsContextSensitive() { return bIsActionMenuContextSensitive; }
 
-	bool OnChangeConsumer( TSharedPtr<FString>& NewConsumer );
-
-	void OnAddProducer( int32 Index );
-
-	void OnRemoveProducer( int32 Index );
-
-	void OnDataprepAssetChanged( UObject* InObject, FPropertyChangedEvent& InPropertyChangedEvent);
-
-	/** Indicates one or all the producers have changed */
-	void OnProducerChanged( int32 Index = INDEX_NONE );
-
 private:
 	void BindCommands();
 	void OnSaveScene();
@@ -124,8 +103,9 @@ private:
 	void CleanPreviewWorld();
 	void OnExecutePipeline();
 	void OnCommitWorld();
+
 	/** Updates asset preview and scene outliner */
-	void OnWorldChanged();
+	void UpdatePreviewPanels();
 
 	void CreateTabs();
 
@@ -164,8 +144,6 @@ private:
 
 	virtual bool OnRequestClose() override;
 
-	void DeleteRegisteredAsset( UObject* Object );
-
 	/** Returns content folder under which all assets are stored after execution of all producers */
 	FString GetTransientContentFolder();
 
@@ -175,6 +153,9 @@ private:
 	/** Recreate preview world from snapshot */
 	void RestoreFromSnapshot();
 
+	/** Handles changes in the Dataprep asset */
+	void OnDataprepAssetChanged(FDataprepAssetChangeType ChangeType, int32 Index );
+
 private:
 	bool bWorldBuilt;
 	bool bIsFirstRun;
@@ -215,9 +196,6 @@ private:
 	/** flag raised to prevent this editor to be closed */
 	bool bIgnoreCloseRequest;
 
-	/** Array of UClasses deriving from UDataprepContentProducer */
-	TArray< DataprepEditorClassDescription > ProducerDescriptions;
-
 	/** Array of UClasses deriving from UDataprepContentConsumer */
 	TArray< DataprepEditorClassDescription > ConsumerDescriptions;
 
diff --git a/Engine/Plugins/Experimental/DataPrepEditor/Source/DataPrepEditor/Private/DataprepDetailsTab.cpp b/Engine/Plugins/Experimental/DataPrepEditor/Source/DataPrepEditor/Private/DataprepDetailsTab.cpp
--- a/Engine/Plugins/Experimental/DataPrepEditor/Source/DataPrepEditor/Private/DataprepDetailsTab.cpp
+++ b/Engine/Plugins/Experimental/DataPrepEditor/Source/DataPrepEditor/Private/DataprepDetailsTab.cpp
@@ -68,7 +68,7 @@ TSharedRef<SDockTab> FDataprepEditor::SpawnTabDetails(const FSpawnTabArgs & Args
 void FDataprepEditor::CreateDetailsViews()
 {
 	DetailsView = SNew( SGraphNodeDetailsWidget );
-	DataprepAssetView = SNew( SDataprepAssetView, this );
+	DataprepAssetView = SNew( SDataprepAssetView, DataprepAssetPtr.Get(), PipelineEditorCommands );
 }
 
 void FDataprepEditor::OnShowSettings()
diff --git a/Engine/Plugins/Experimental/DataPrepEditor/Source/DataPrepEditor/Private/Widgets/DataprepAssetView.cpp b/Engine/Plugins/Experimental/DataPrepEditor/Source/DataPrepEditor/Private/Widgets/DataprepAssetView.cpp
--- a/Engine/Plugins/Experimental/DataPrepEditor/Source/DataPrepEditor/Private/Widgets/DataprepAssetView.cpp
+++ b/Engine/Plugins/Experimental/DataPrepEditor/Source/DataPrepEditor/Private/Widgets/DataprepAssetView.cpp
@@ -16,6 +16,7 @@
 #include "K2Node_AddComponent.h"
 #include "PropertyEditorModule.h"
 
+#include "EditorFontGlyphs.h"
 #include "Modules/ModuleManager.h"
 #include "Widgets/Colors/SColorBlock.h"
 #include "Widgets/Images/SImage.h"
@@ -50,8 +51,6 @@ void SGraphNodeDetailsWidget::Construct(const FArguments& InArgs)
 
 	PropertyView = EditModule.CreateDetailView(DetailsViewArgs);
 
-	//PropertyView->OnFinishedChangingProperties().Add(FOnFinishedChangingProperties::FDelegate::CreateSP(this, &FBlueprintEditor::OnFinishedChangingProperties));
-
 	PropertyView->GetIsPropertyEditingEnabledDelegate().BindSP(this, &SGraphNodeDetailsWidget::GetCanEditProperties);
 
 	// Create the border that all of the content will get stuffed into
@@ -267,18 +266,8 @@ void SGraphNodeDetailsWidget::UpdateFromObjects(const TArray<UObject*>& Property
 		.VAlign(VAlign_Top)
 		[
 			SNew(SBox)
-			//.Visibility(this, &SKismetInspector::GetPropertyViewVisibility)
 			[
 				SNew(SVerticalBox)
-				//+ SVerticalBox::Slot()
-				//.AutoHeight()
-				//.Padding(FMargin(0, 0, 0, 1))
-				//[
-				//	SNew(SKismetInspectorUneditableComponentWarning)
-				//	.Visibility(this, &SKismetInspector::GetInheritedBlueprintComponentWarningVisibility)
-				//	.WarningText(NSLOCTEXT("SKismetInspector", "BlueprintUneditableInheritedComponentWarning", "Components flagged as not editable when inherited must be edited in the <a id=\"HyperlinkDecorator\" style=\"DetailsView.BPMessageHyperlinkStyle\">Parent Blueprint</>"))
-				//	.OnHyperlinkClicked(this, &SKismetInspector::OnInheritedBlueprintComponentWarningHyperlinkClicked)
-				//]
 				+ SVerticalBox::Slot()
 				[
 					PropertyView.ToSharedRef()
@@ -290,10 +279,10 @@ void SGraphNodeDetailsWidget::UpdateFromObjects(const TArray<UObject*>& Property
 	ContextualEditingBorderWidget->SetContent(ContextualEditingWidget);
 }
 
-void SProducerStackEntryTreeView::Construct(const FArguments& InArgs, SDataprepAssetView* InDataprepAssetView, FDataprepEditor* InDataprepEditorPtr)
+void SProducerStackEntryTreeView::Construct(const FArguments& InArgs, SDataprepAssetView* InDataprepAssetView, UDataprepAsset* InDataprepAssetPtr)
 {
-	DataprepEditorPtr = InDataprepEditorPtr;
-	DataprepEditorPtr->OnDataprepAssetProducerChanged().AddSP(this, &SProducerStackEntryTreeView::OnDataprepAssetProducerChanged);
+	DataprepAssetPtr = InDataprepAssetPtr;
+	check( DataprepAssetPtr.IsValid() );
 
 	BuildProducerEntries();
 
@@ -328,16 +317,19 @@ TSharedRef<ITableRow> SProducerStackEntryTreeView::OnGenerateRow(FProducerStackE
 
 void SProducerStackEntryTreeView::BuildProducerEntries()
 {
-	const TArray< FDataprepAssetProducer >& Producers = DataprepEditorPtr->GetDataprepAsset()->Producers;
+	if( UDataprepAsset* DataprepAsset = DataprepAssetPtr.Get() )
+	{
+		int32 ProducersCount = DataprepAsset->GetProducersCount();
 
-	RootNodes.Empty( Producers.Num() );
+		RootNodes.Empty( ProducersCount );
 
-	for ( const FDataprepAssetProducer& AssetProducer : Producers )
-	{
-		if (AssetProducer.Producer != nullptr)
+		for( int32 Index = 0; Index < ProducersCount; ++Index )
 		{
-			TSharedRef<FProducerStackEntry> ProducerStackEntry = MakeShareable( new FProducerStackEntry( RootNodes.Num(), AssetProducer, DataprepEditorPtr ) );
-			RootNodes.Add( ProducerStackEntry );
+			if(DataprepAsset->GetProducer( Index ) != nullptr)
+			{
+				TSharedRef<FProducerStackEntry> ProducerStackEntry = MakeShareable( new FProducerStackEntry( Index, DataprepAsset ) );
+				RootNodes.Add( ProducerStackEntry );
+			}
 		}
 	}
 }
@@ -380,19 +372,17 @@ TSharedRef<SWidget> SProducerStackEntryTableRow::GetInputMainWidget()
 		return SNullWidget::NullWidget;
 	}
 
-	auto CheckEntry = [this, ProducerStackEntry]()
+	auto ToggleEntry = [this, ProducerStackEntry]()
 	{
-		ProducerStackEntry->bIsEnabled = !ProducerStackEntry->bIsEnabled;
-		CheckBox->SetText( FText::FromString( ProducerStackEntry->bIsEnabled ? TEXT( "\xf14a" ) : TEXT( "\xf0c8" ) ) );
-
-		ProducerStackEntry->GetDataprepAssetProducer().bIsEnabled = ProducerStackEntry->bIsEnabled;
+		ProducerStackEntry->ToggleProducer();
+		CheckBox->SetText( ProducerStackEntry->WillBeRun() ? FEditorFontGlyphs::Check_Square : FEditorFontGlyphs::Square );
 
 		return FReply::Handled();
 	};
 
 	auto DeleteEntry = [ProducerStackEntry]()
 	{
-		ProducerStackEntry->DataprepEditorPtr->OnRemoveProducer( ProducerStackEntry->ProducerIndex );
+		ProducerStackEntry->RemoveProducer();
 		return FReply::Handled();
 	};
 
@@ -401,6 +391,7 @@ TSharedRef<SWidget> SProducerStackEntryTableRow::GetInputMainWidget()
 
 	TSharedPtr<SWidget> Widget = SNew(SBorder)
 	.BorderImage(FEditorStyle::GetBrush("NoBrush"))
+	.BorderBackgroundColor( ProducerStackEntry->bIsSuperseded ? FLinearColor::White : FLinearColor::Red )
 	.Padding(5.0f)
 	[
 		SNew(SHorizontalBox)
@@ -414,14 +405,14 @@ TSharedRef<SWidget> SProducerStackEntryTableRow::GetInputMainWidget()
 			.ButtonStyle(FEditorStyle::Get(), "HoverHintOnly")
 			.ToolTipText(LOCTEXT("CheckToolTip", "Include or exclude this producer to the creation of the world "))
 			.IsFocusable(false)
-			.OnClicked_Lambda( CheckEntry )
+			.OnClicked_Lambda( ToggleEntry )
 			.VAlign( VAlign_Top )
 			.Content()
 			[
 				SAssignNew(CheckBox, STextBlock)
 				.Font(FEditorStyle::Get().GetFontStyle(DataprepTabsFontName))
 				.ColorAndOpacity( FLinearColor::White )
-				.Text(FText::FromString(FString( ProducerStackEntry->bIsEnabled ? TEXT( "\xf14a" ) : TEXT( "\xf0c8" ) ) ) )
+				.Text( ProducerStackEntry->WillBeRun() ? FEditorFontGlyphs::Check_Square : FEditorFontGlyphs::Square )
 			]
 		]
 		// Input entry label
@@ -448,58 +439,73 @@ TSharedRef<SWidget> SProducerStackEntryTableRow::GetInputMainWidget()
 				SNew(STextBlock)
 				.Font(FEditorStyle::Get().GetFontStyle(DataprepTabsFontName))
 				.ColorAndOpacity( FLinearColor::White )
-				.Text(FText::FromString(FString(TEXT("\xf1f8"))))
+				.Text( FEditorFontGlyphs::Trash )
 			]
 		]
 	];
 
+	// Disable check box if associated producer is superseeded
+	CheckBox->SetEnabled( !ProducerStackEntry->bIsSuperseded );
+
 	return Widget.ToSharedRef();
 }
 
-void SDataprepAssetView::Construct( const FArguments& InArgs, FDataprepEditor* InDataprepEditor )
+void SDataprepAssetView::Construct( const FArguments& InArgs, UDataprepAsset* InDataprepAssetPtr, TSharedPtr<FUICommandList>& CommandList )
 {
-	DataprepEditorPtr = InDataprepEditor;
+	check( InDataprepAssetPtr );
 
-	const UDataprepAsset* DataprepAsset = DataprepEditorPtr->GetDataprepAsset();
+	DataprepAssetPtr = InDataprepAssetPtr;
+
+	DataprepAssetPtr->GetOnChanged().AddRaw( this, &SDataprepAssetView::OnDataprepAssetChanged );
 
 	bIsChecked = true;
-	const TArray< FDataprepAssetProducer >& Producers = DataprepAsset->Producers;
-	for (const FDataprepAssetProducer& Producer : Producers)
+
+	for(int32 Index = 0; Index < DataprepAssetPtr->GetProducersCount(); ++Index)
 	{
-		bIsChecked &= Producer.bIsEnabled;
+		bIsChecked &= DataprepAssetPtr->IsProducerEnabled( Index ) && !DataprepAssetPtr->IsProducerSuperseded( Index );
 	}
 
-	const TArray< DataprepEditorClassDescription >& Consumers = DataprepEditorPtr->GetConsumerDescriptions();
-
-	for (const DataprepEditorClassDescription& Consumer : Consumers)
+	for( TObjectIterator< UClass > It ; It ; ++It )
 	{
-		TSharedPtr< FString >& ConsumerLabel = ConsumerList.Emplace_GetRef( new FString( Consumer.Get<1>().ToString() ) );
-		if (DataprepAsset->Consumer != nullptr && DataprepAsset->Consumer->GetClass() == Consumer.Get<0>() )
+		UClass* CurrentClass = (*It);
+
+		if ( !CurrentClass->HasAnyClassFlags( CLASS_Abstract ) )
 		{
-			SelectedProducer = ConsumerLabel;
+			if( CurrentClass->IsChildOf( UDataprepContentConsumer::StaticClass() ) )
+			{
+				if( UDataprepContentConsumer* Consumer = Cast< UDataprepContentConsumer >( CurrentClass->GetDefaultObject() ) )
+				{
+					TSharedPtr< FString >& ConsumerDescriptionLabel = ConsumerDescriptionList.Emplace_GetRef( new FString( Consumer->GetLabel().ToString() ) );
+					ConsumerDescriptionMap.Add( ConsumerDescriptionLabel, CurrentClass );
+
+					if (DataprepAssetPtr->GetConsumer() != nullptr && DataprepAssetPtr->GetConsumer()->GetClass() == CurrentClass )
+					{
+						SelectedConsumerDescription = ConsumerDescriptionLabel;
+					}
+				}
+			}
 		}
 	}
 
-	if ( !SelectedProducer.IsValid() && ConsumerList.Num() > 0)
+	if ( !SelectedConsumerDescription.IsValid() && ConsumerDescriptionList.Num() > 0)
 	{
-		SelectedProducer = MakeShared<FString>( FString() );
+		SelectedConsumerDescription = MakeShared<FString>( FString() );
 	}
 
 	auto CheckEntry = [ this ]()
 	{
-		bIsChecked = !bIsChecked;
+		for(int32 Index = 0; Index < DataprepAssetPtr->GetProducersCount(); ++Index)
+		{
+			DataprepAssetPtr->EnableProducer( Index, !bIsChecked );
+		}
 
-		CheckBox->SetText(FText::FromString( FString( bIsChecked ? TEXT( "\xf14a" ) : TEXT( "\xf0c8" ) ) ) );
 
-		TArray< FDataprepAssetProducer >& Producers = DataprepEditorPtr->GetDataprepAsset()->Producers;
-		for ( FDataprepAssetProducer& Producer : Producers )
+		if( DataprepAssetPtr->EnableAllProducers( !bIsChecked ) )
 		{
-			Producer.bIsEnabled = bIsChecked;
+			bIsChecked = !bIsChecked;
 		}
 
-		// #ueent_todo: Revisit how changes are propagated from UI to editor to asset and reverse
-		// Propagate dataprep asset has changed
-		DataprepEditorPtr->OnDataprepAssetProducerChanged().Broadcast();
+		CheckBox->SetText( bIsChecked ? FEditorFontGlyphs::Check_Square : FEditorFontGlyphs::Square );
 
 		TreeView->Refresh();
 
@@ -510,7 +516,7 @@ void SDataprepAssetView::Construct( const FArguments& InArgs, FDataprepEditor* I
 	.ComboButtonStyle(FEditorStyle::Get(), "ToolbarComboButton")
 	.ForegroundColor(FLinearColor::White)
 	.ToolTipText(LOCTEXT("AddNewToolTip", "Add a new producer."))
-	.OnGetMenuContent(this, &SDataprepAssetView::CreateAddProducerMenuWidget)
+	.OnGetMenuContent(this, &SDataprepAssetView::CreateAddProducerMenuWidget, CommandList)
 	.HasDownArrow(false)
 	.ButtonContent()
 	[
@@ -527,7 +533,7 @@ void SDataprepAssetView::Construct( const FArguments& InArgs, FDataprepEditor* I
 		]
 	];
 
-	TreeView = SNew(SProducerStackEntryTreeView, this, DataprepEditorPtr);
+	TreeView = SNew(SProducerStackEntryTreeView, this, DataprepAssetPtr.Get() );
 
 	// #ueent_todo: Look at changing the border brushes to add color to this stuff
 	ChildSlot
@@ -569,7 +575,6 @@ void SDataprepAssetView::Construct( const FArguments& InArgs, FDataprepEditor* I
 				.Padding( 5.0f )
 				[
 					SNew(STextBlock)
-					.Font(FEditorStyle::Get().GetFontStyle(DataprepTabsFontName))
 					.Text(LOCTEXT("DataprepAssetView_Consumer_label", "Output"))
 					.MinDesiredWidth( 200 )
 				]
@@ -584,16 +589,16 @@ void SDataprepAssetView::Construct( const FArguments& InArgs, FDataprepEditor* I
 				.VAlign(VAlign_Center)
 				[
 					SAssignNew( ProducerSelector, STextComboBox )
-					.OptionsSource( &ConsumerList )
-					.OnSelectionChanged( this, &SDataprepAssetView::OnConsumerChanged )
-					.InitiallySelectedItem( SelectedProducer )
+					.OptionsSource( &ConsumerDescriptionList )
+					.OnSelectionChanged( this, &SDataprepAssetView::OnNewConsumerSelected )
+					.InitiallySelectedItem( SelectedConsumerDescription )
 				]
 			]
 			+ SVerticalBox::Slot()
 			.AutoHeight()
 			[
 				SAssignNew( ConsumerWidget, SDataprepConsumerWidget )
-				.DataprepConsumer( DataprepAsset->Consumer )
+				.DataprepConsumer( DataprepAssetPtr->GetConsumer() )
 			]
 			// Section for producers
 			+ SVerticalBox::Slot()
@@ -641,7 +646,7 @@ void SDataprepAssetView::Construct( const FArguments& InArgs, FDataprepEditor* I
 						SAssignNew(CheckBox, STextBlock)
 						.Font(FEditorStyle::Get().GetFontStyle(DataprepTabsFontName))
 						.ColorAndOpacity( FLinearColor::White )
-						.Text(FText::FromString(FString( bIsChecked ? TEXT( "\xf14a" ) : TEXT( "\xf0c8" ) ) ) )
+						.Text( bIsChecked ? FEditorFontGlyphs::Check_Square : FEditorFontGlyphs::Square )
 					]
 				]
 				+ SHorizontalBox::Slot()
@@ -650,7 +655,6 @@ void SDataprepAssetView::Construct( const FArguments& InArgs, FDataprepEditor* I
 				.VAlign(VAlign_Center)
 				[
 					SNew(STextBlock)
-					.Font(FEditorStyle::Get().GetFontStyle(DataprepTabsFontName))
 					.Text(LOCTEXT("DataprepAssetView_Producers_label", "Inputs"))
 				]
 				+ SHorizontalBox::Slot()
@@ -676,54 +680,93 @@ void SDataprepAssetView::Construct( const FArguments& InArgs, FDataprepEditor* I
 	];
 }
 
-void SDataprepAssetView::OnConsumerChanged( TSharedPtr<FString> NewConsumer, ESelectInfo::Type SelectInfo)
+SDataprepAssetView::~SDataprepAssetView()
 {
-	if ( !NewConsumer.IsValid() )
+	if( UDataprepAsset* DataprepAsset = DataprepAssetPtr.Get() )
 	{
-		return;
+		DataprepAsset->GetOnChanged().RemoveAll( this );
 	}
+}
 
-	if ( SelectedProducer == NewConsumer )
+void SDataprepAssetView::OnNewConsumerSelected( TSharedPtr<FString> NewConsumerDescription, ESelectInfo::Type SelectInfo)
+{
+	if ( !NewConsumerDescription.IsValid() || SelectedConsumerDescription == NewConsumerDescription )
 	{
 		return;
 	}
 
-	// #ueent_todo: Provide an index instead of a FString
-	if ( !DataprepEditorPtr->OnChangeConsumer( NewConsumer ) )
+	if( UDataprepAsset* DataprepAsset = DataprepAssetPtr.Get() )
 	{
-		ProducerSelector->SetSelectedItem( SelectedProducer );
+		UClass** NewConsumerClassPtr = ConsumerDescriptionMap.Find(NewConsumerDescription);
+		check(NewConsumerClassPtr);
+
+		if( !DataprepAsset->ReplaceConsumer( *NewConsumerClassPtr ) )
+		{
+			ProducerSelector->SetSelectedItem(SelectedConsumerDescription);
+		}
+		// Update SelectedConsumerDescription only, the widget displaying the consumer is updated thru notifications 
+		else
+		{
+			SelectedConsumerDescription = NewConsumerDescription;
+		}
 	}
-	else
+}
+
+void SDataprepAssetView::OnDataprepAssetChanged(FDataprepAssetChangeType ChangeType, int32 Index)
+{
+	if( UDataprepAsset* DataprepAsset = DataprepAssetPtr.Get() )
 	{
-		SelectedProducer = NewConsumer;
-		ConsumerWidget->SetDataprepConsumer( DataprepEditorPtr->GetDataprepAsset()->Consumer );
+		if(ChangeType == FDataprepAssetChangeType::ConsumerModified)
+		{
+			// Update the widget holding the consumer
+			ConsumerWidget->SetDataprepConsumer( DataprepAsset->GetConsumer() );
+		}
+		else if( ChangeType == FDataprepAssetChangeType::ProducerModified ||
+			ChangeType == FDataprepAssetChangeType::ProducerAdded ||
+			ChangeType == FDataprepAssetChangeType::ProducerRemoved )
+		{
+			// Brute force : Regenerate the whole tree view
+			TreeView->Refresh();
+		}
 	}
 }
 
-TSharedRef<SWidget> SDataprepAssetView::CreateAddProducerMenuWidget()
+TSharedRef<SWidget> SDataprepAssetView::CreateAddProducerMenuWidget(TSharedPtr<FUICommandList> CommandList)
 {
 	const bool bShouldCloseWindowAfterMenuSelection = true;
-	FMenuBuilder MenuBuilder(bShouldCloseWindowAfterMenuSelection, DataprepEditorPtr->GetToolkitCommands());
+	FMenuBuilder MenuBuilder(bShouldCloseWindowAfterMenuSelection, CommandList);
 
 	MenuBuilder.BeginSection("AddNewProducer", LOCTEXT("DataprepEditorViews_AddImports", "Add Producer"));
 	{
 		FUIAction MenuAction;
 		int32 Index = 0;
 
-		for (const DataprepEditorClassDescription& ProducerDescription : DataprepEditorPtr->GetProducerDescriptions())
+		// Find content producers the user could use for his/her data preparation
+		for( TObjectIterator< UClass > It ; It ; ++It )
 		{
-			MenuAction.ExecuteAction.BindSP(DataprepEditorPtr, &FDataprepEditor::OnAddProducer, Index);
-
-			MenuBuilder.AddMenuEntry(
-				ProducerDescription.Get<1>(),
-				ProducerDescription.Get<2>(),
-				FSlateIcon( FDataprepEditorStyle::GetStyleSetName(), TEXT("DataprepEditor.Producer") ),
-				MenuAction,
-				NAME_None,
-				EUserInterfaceActionType::Button
-			);
-
-			++Index;
+			UClass* CurrentClass = (*It);
+
+			if ( !CurrentClass->HasAnyClassFlags( CLASS_Abstract ) )
+			{
+				if( CurrentClass->IsChildOf( UDataprepContentProducer::StaticClass() ) )
+				{
+					MenuAction.ExecuteAction.BindSP(this, &SDataprepAssetView::OnAddProducer, CurrentClass);
+
+					UDataprepContentProducer* DefaultProducer = CurrentClass->GetDefaultObject<UDataprepContentProducer>();
+					check( DefaultProducer );
+
+					MenuBuilder.AddMenuEntry(
+						DefaultProducer->GetLabel(),
+						DefaultProducer->GetDescription(),
+						FSlateIcon( FDataprepEditorStyle::GetStyleSetName(), TEXT("DataprepEditor.Producer") ),
+						MenuAction,
+						NAME_None,
+						EUserInterfaceActionType::Button
+					);
+
+					++Index;
+				}
+			}
 		}
 	}
 	MenuBuilder.EndSection();
@@ -731,28 +774,27 @@ TSharedRef<SWidget> SDataprepAssetView::CreateAddProducerMenuWidget()
 	return MenuBuilder.MakeWidget();
 }
 
-void SDataprepAssetView::OnSelectionChanged( TSharedPtr< FProducerStackEntry > InItem, ESelectInfo::Type InSeletionInfo, FDataprepEditor* DataprepEditor )
+void SDataprepAssetView::OnAddProducer( UClass* ProducerClass )
 {
-	// An entry is selected
-	if ( InItem.IsValid() )
+	if( UDataprepAsset* DataprepAsset = DataprepAssetPtr.Get() )
 	{
-		const TArray< FDataprepAssetProducer >& Producers = DataprepEditorPtr->GetDataprepAsset()->Producers;
+		DataprepAsset->AddProducer(ProducerClass);
+	}
+}
 
-		if (Producers.IsValidIndex(InItem->ProducerIndex))
-		{
-			// Take hold on the selected entry
-			SelectedEntry = InItem;
-			// Show properties of associated producer
-			DetailsView->SetObject( Producers[InItem->ProducerIndex].Producer, true );
-		}
+void SDataprepAssetView::OnSelectionChanged( TSharedPtr< FProducerStackEntry > InItem, ESelectInfo::Type InSeletionInfo )
+{
+	// An entry is selected
+	if ( InItem.IsValid() && InItem->HasValidData() )
+	{
+		// Take hold on the selected entry
+		SelectedEntry = InItem;
 	}
 	// An entry is deselected
 	else
 	{
 		// Release hold on selected entry
 		SelectedEntry.Reset();
-		// Clear details section
-		DetailsView->SetObject( nullptr, true );
 	}
 }
 
diff --git a/Engine/Plugins/Experimental/DataPrepEditor/Source/DataPrepEditor/Private/Widgets/DataprepAssetView.h b/Engine/Plugins/Experimental/DataPrepEditor/Source/DataPrepEditor/Private/Widgets/DataprepAssetView.h
--- a/Engine/Plugins/Experimental/DataPrepEditor/Source/DataPrepEditor/Private/Widgets/DataprepAssetView.h
+++ b/Engine/Plugins/Experimental/DataPrepEditor/Source/DataPrepEditor/Private/Widgets/DataprepAssetView.h
@@ -21,36 +21,59 @@ class SDataprepConsumerWidget;
 class FProducerStackEntry
 {
 public:
-	FProducerStackEntry(int32 ProducerIndex, const FDataprepAssetProducer& Producer, FDataprepEditor* InDataprepEditor)
-		: Label( Producer.Producer->GetLabel().ToString() )
-		, ProducerIndex( ProducerIndex )
-		, bIsEnabled( Producer.bIsEnabled )
-		, bIsSelected( false )
-		, DataprepEditorPtr( InDataprepEditor )
+	FProducerStackEntry(int32 InProducerIndex, UDataprepAsset* InDataprepAssetPtr)
+		: ProducerIndex( InProducerIndex )
+		, bIsEnabled( false )
+		, bIsSuperseded( false )
+		, DataprepAssetPtr( InDataprepAssetPtr )
 	{
+		if( UDataprepAsset* DataprepAsset = DataprepAssetPtr.Get() )
+		{
+			if( const UDataprepContentProducer* Producer = DataprepAsset->GetProducer( ProducerIndex ) )
+			{
+				bIsEnabled = DataprepAsset->IsProducerEnabled( ProducerIndex );
+				bIsSuperseded = DataprepAsset->IsProducerSuperseded( ProducerIndex );
+				Label = Producer->GetLabel().ToString();
+			}
+		}
 	}
 
 	bool HasValidData()
 	{
-		return DataprepEditorPtr->GetDataprepAsset()->Producers.IsValidIndex( ProducerIndex );
+		return DataprepAssetPtr.IsValid() && DataprepAssetPtr->GetProducer( ProducerIndex ) != nullptr;
 	}
 
 	UDataprepContentProducer* GetProducer()
 	{
-		return HasValidData() ? DataprepEditorPtr->GetDataprepAsset()->Producers[ProducerIndex].Producer : nullptr;
+		return DataprepAssetPtr.IsValid() ? DataprepAssetPtr->GetProducer( ProducerIndex ) : nullptr;
 	}
 
-	FDataprepAssetProducer& GetDataprepAssetProducer()
+	bool WillBeRun() { return bIsEnabled && !bIsSuperseded; }
+
+	void ToggleProducer()
+	{
+		if( UDataprepAsset* DataprepAsset = DataprepAssetPtr.Get() )
+		{
+			DataprepAsset->EnableProducer(ProducerIndex, !bIsEnabled);
+
+			// #ueent_todo: Cache previous value to report failed enabling/disabling
+			bIsEnabled = DataprepAsset->IsProducerEnabled(ProducerIndex);
+		}
+	}
+
+	void RemoveProducer()
 	{
-		// #ueent_todo: Add defensive code
-		return DataprepEditorPtr->GetDataprepAsset()->Producers[ProducerIndex];
+		if( UDataprepAsset* DataprepAsset = DataprepAssetPtr.Get() )
+		{
+			DataprepAsset->RemoveProducer( ProducerIndex );
+		}
 	}
 
 	FString Label;
 	int32 ProducerIndex;
 	bool bIsEnabled;
-	bool bIsSelected;
-	FDataprepEditor* DataprepEditorPtr;
+	bool bIsSuperseded;
+	TWeakObjectPtr<UDataprepAsset> DataprepAssetPtr;
 };
 
 typedef TSharedRef<FProducerStackEntry> FProducerStackEntryRef;
@@ -84,7 +107,7 @@ public:
 	SLATE_BEGIN_ARGS(SProducerStackEntryTreeView) {}
 	SLATE_END_ARGS()
 
-	void Construct(const FArguments& InArgs, SDataprepAssetView* InDataprepAssetView, FDataprepEditor* InDataprepEditorPtr);
+	void Construct(const FArguments& InArgs, SDataprepAssetView* InDataprepAssetView, UDataprepAsset* InDataprepAssetPtr);
 
 	int32 GetDisplayIndexOfNode(FProducerStackEntryRef InNode);
 
@@ -101,7 +124,7 @@ private:
 	void OnDataprepAssetProducerChanged();
 
 private:
-	FDataprepEditor* DataprepEditorPtr;
+	TWeakObjectPtr<UDataprepAsset> DataprepAssetPtr;
 	TArray<FProducerStackEntryRef> RootNodes;
 };
 
@@ -111,26 +134,30 @@ public:
 	SLATE_BEGIN_ARGS(SDataprepAssetView) {}
 	SLATE_END_ARGS()
 
-	void Construct(const FArguments& InArgs, FDataprepEditor* InDataprepEditor);
+	void Construct(const FArguments& InArgs, UDataprepAsset* InDataprepAssetPtr, TSharedPtr<FUICommandList>& CommandList);
 
-	~SDataprepAssetView()
-	{
-	}
+	~SDataprepAssetView();
 
 	/** Set the renderer when we change selection */
-	void OnSelectionChanged( TSharedPtr< FProducerStackEntry > InItem, ESelectInfo::Type InSeletionInfo, FDataprepEditor* DataprepEditor );
+	void OnSelectionChanged( TSharedPtr< FProducerStackEntry > InItem, ESelectInfo::Type InSeletionInfo );
 
 private:
-	TSharedRef<SWidget> CreateAddProducerMenuWidget();
-	void OnConsumerChanged( TSharedPtr<FString> NewConsumer, ESelectInfo::Type SelectInfo);
+	TSharedRef<SWidget> CreateAddProducerMenuWidget(TSharedPtr<FUICommandList> CommandList);
+
+	void OnNewConsumerSelected( TSharedPtr<FString> NewConsumer, ESelectInfo::Type SelectInfo);
+
+	void OnAddProducer( UClass* ProducerClass );
+
+	/** Handles changes in the Dataprep asset */
+	void OnDataprepAssetChanged(FDataprepAssetChangeType ChangeType, int32 Index );
 
 private:
-	FDataprepEditor* DataprepEditorPtr;
+	TWeakObjectPtr<UDataprepAsset> DataprepAssetPtr;
 	TSharedPtr<SProducerStackEntryTreeView> TreeView;
-	TSharedPtr<IDetailsView> DetailsView;
 	TSharedPtr< STextBlock > CheckBox;
-	TArray< TSharedPtr< FString > > ConsumerList;
-	TSharedPtr< FString > SelectedProducer;
+	TArray< TSharedPtr< FString > > ConsumerDescriptionList;
+	TMap< TSharedPtr< FString >, UClass* > ConsumerDescriptionMap;
+	TSharedPtr< FString > SelectedConsumerDescription;
 	TSharedPtr< STextComboBox > ProducerSelector;
 	bool bIsChecked;
 	TSharedPtr< FProducerStackEntry > SelectedEntry;
diff --git a/Engine/Plugins/Experimental/DataPrepEditor/Source/DataPrepEditor/Private/Widgets/SAssetsPreviewWidget.cpp b/Engine/Plugins/Experimental/DataPrepEditor/Source/DataPrepEditor/Private/Widgets/SAssetsPreviewWidget.cpp
--- a/Engine/Plugins/Experimental/DataPrepEditor/Source/DataPrepEditor/Private/Widgets/SAssetsPreviewWidget.cpp
+++ b/Engine/Plugins/Experimental/DataPrepEditor/Source/DataPrepEditor/Private/Widgets/SAssetsPreviewWidget.cpp
@@ -143,9 +143,10 @@ namespace AssetPreviewWidget
 		];
 	}
 
-	void SAssetsPreviewWidget::SetAssetsList(const TArray< TWeakObjectPtr< UObject > >& InAssetsList, const FString& InPathToTempFolder)
+	void SAssetsPreviewWidget::SetAssetsList(const TArray< TWeakObjectPtr< UObject > >& InAssetsList, const FString& InPathToReplace, const FString& InSubstitutePath)
 	{
-		PathToTempFolder = InPathToTempFolder;
+		PathToReplace = InPathToReplace;
+		SubstitutePath = InSubstitutePath;
 		RootItems.Empty();
 		TMap< FString, FAssetTreeItemPtr > NamesToRootItem;
 
@@ -264,7 +265,10 @@ namespace AssetPreviewWidget
 	{
 		TArray< FString > ItemsName;
 		FString AssetSubPath = Asset->GetPathName(nullptr);
-		AssetSubPath.RemoveFromStart(PathToTempFolder);
+		if(AssetSubPath.RemoveFromStart(PathToReplace) && !SubstitutePath.IsEmpty())
+		{
+			AssetSubPath = SubstitutePath / AssetSubPath;
+		}
 		AssetSubPath.ReplaceCharInline(TEXT('/'), TEXT('.'));
 		AssetSubPath.ParseIntoArray(ItemsName, TEXT("."), true);
 		return ItemsName;
diff --git a/Engine/Plugins/Experimental/DataPrepEditor/Source/DataPrepEditor/Private/Widgets/SAssetsPreviewWidget.h b/Engine/Plugins/Experimental/DataPrepEditor/Source/DataPrepEditor/Private/Widgets/SAssetsPreviewWidget.h
--- a/Engine/Plugins/Experimental/DataPrepEditor/Source/DataPrepEditor/Private/Widgets/SAssetsPreviewWidget.h
+++ b/Engine/Plugins/Experimental/DataPrepEditor/Source/DataPrepEditor/Private/Widgets/SAssetsPreviewWidget.h
@@ -28,7 +28,7 @@ namespace AssetPreviewWidget
 	public:
 		void Construct(const FArguments& InArgs);
 
-		void SetAssetsList(const TArray< TWeakObjectPtr< UObject > >& InAssetsList, const FString& InPathToTempFolder);
+		void SetAssetsList(const TArray< TWeakObjectPtr< UObject > >& InAssetsList, const FString& InPathToReplace, const FString& InSubstitutePath);
 		void ClearAssetList();
 
 		FOnSelectionChanged& OnSelectionChanged() { return OnSelectionChangedDelegate; };
@@ -60,7 +60,8 @@ namespace AssetPreviewWidget
 		TSharedPtr< STreeView< FAssetTreeItemPtr > > TreeView;
 
 		FText FilterText;
-		FString PathToTempFolder;
+		FString PathToReplace;
+		FString SubstitutePath;
 
 		FOnSelectionChanged OnSelectionChangedDelegate;
 	};
