Copying CL #6957158 and 6969370 from //UE4/Dev-Enterprise/...

Fixed crash when importing and commiting twice
Fixed issue with inability to select the root content folder
Changed the display of the targeted content folder to start with /Content instead of /Game in the UI
Fixed some of the wrong labels mentionned in UE-75968 except the name of dataprep asset newly created
Fixed issue with Umaterial objects not imported if not instancieted
Removed incomplete File write consumer

diff --git a/Engine/Plugins/Experimental/DataPrepEditor/Source/DataPrepCore/Private/DataPrepContentConsumer.cpp b/Engine/Plugins/Experimental/DataPrepEditor/Source/DataPrepCore/Private/DataPrepContentConsumer.cpp
--- a/Engine/Plugins/Experimental/DataPrepEditor/Source/DataPrepCore/Private/DataPrepContentConsumer.cpp
+++ b/Engine/Plugins/Experimental/DataPrepEditor/Source/DataPrepCore/Private/DataPrepContentConsumer.cpp
@@ -50,7 +50,9 @@ bool UDataprepContentConsumer::SetTargetContentFolder(const FString& InTargetCon
 		return true;
 	}
 
-	if( FPackageName::IsValidLongPackageName( InTargetContentFolder ) )
+	// Pretend creating a dummy package to verify packages could be created under this content folder.
+	FString LongPackageName = InTargetContentFolder / TEXT("DummyPackageName");
+	if( FPackageName::IsValidLongPackageName( LongPackageName ) )
 	{
 		TargetContentFolder = InTargetContentFolder;
 		OnChanged.Broadcast();
diff --git a/Engine/Plugins/Experimental/DataPrepEditor/Source/DataPrepEditor/Private/AssetTypeActions_DataPrepAsset.cpp b/Engine/Plugins/Experimental/DataPrepEditor/Source/DataPrepEditor/Private/AssetTypeActions_DataPrepAsset.cpp
--- a/Engine/Plugins/Experimental/DataPrepEditor/Source/DataPrepEditor/Private/AssetTypeActions_DataPrepAsset.cpp
+++ b/Engine/Plugins/Experimental/DataPrepEditor/Source/DataPrepEditor/Private/AssetTypeActions_DataPrepAsset.cpp
@@ -16,7 +16,7 @@ uint32 FAssetTypeActions_DataprepAsset::GetCategories()
 
 FText FAssetTypeActions_DataprepAsset::GetName() const
 {
-	return NSLOCTEXT("AssetTypeActions_DataprepAsset", "AssetTypeActions_DataprepAsset_Name", "Dataprep Asset");
+	return NSLOCTEXT("AssetTypeActions_DataprepAsset", "AssetTypeActions_DataprepAsset_Name", "Dataprep");
 }
 
 UClass* FAssetTypeActions_DataprepAsset::GetSupportedClass() const
diff --git a/Engine/Plugins/Experimental/DataPrepEditor/Source/DataPrepEditor/Private/DataPrepEditor.cpp b/Engine/Plugins/Experimental/DataPrepEditor/Source/DataPrepEditor/Private/DataPrepEditor.cpp
--- a/Engine/Plugins/Experimental/DataPrepEditor/Source/DataPrepEditor/Private/DataPrepEditor.cpp
+++ b/Engine/Plugins/Experimental/DataPrepEditor/Source/DataPrepEditor/Private/DataPrepEditor.cpp
@@ -580,14 +580,17 @@ void FDataprepEditor::CleanPreviewWorld()
 			if (Actor && !Actor->IsPendingKill() && !DefaultActorsInPreviewWorld.Contains(Actor))
 			{
 				PreviewWorld->EditorDestroyActor(Actor, true);
-				// Since deletion can be delayed, rename to avoid future name collision 
-				Actor->Rename(nullptr, GetTransientPackage(), REN_DontCreateRedirectors | REN_ForceNoResetLoaders);
+
+				// Since deletion can be delayed, rename to avoid future name collision
+				// Call UObject::Rename directly on actor to avoid AActor::Rename which unnecessarily sunregister and re-register components
+				Actor->UObject::Rename(nullptr, GetTransientPackage(), REN_DontCreateRedirectors | REN_ForceNoResetLoaders);
 			}
 		}
 	}
 
+	// Delete assets which are still in the transient content folder
 	const FString TransientContentFolder( GetTransientContentFolder() );
-	CachedAssets.Append( Assets );
+	TArray<UObject*> ObjectsToDelete;
 	for( TWeakObjectPtr<UObject>& Asset : CachedAssets )
 	{
 		if( UObject* ObjectToDelete = Asset.Get() )
@@ -595,22 +598,21 @@ void FDataprepEditor::CleanPreviewWorld()
 			FString PackagePath = ObjectToDelete->GetOutermost()->GetName();
 			if( PackagePath.StartsWith( TransientContentFolder ) )
 			{
-				FName ObjectName = MakeUniqueObjectName( GetTransientPackage(), ObjectToDelete->GetClass(), *( ObjectToDelete->GetName() + TEXT( "_Deleted" ) ) );
-				ObjectToDelete->Rename( *ObjectName.ToString(), GetTransientPackage(), REN_DontCreateRedirectors | REN_NonTransactional );
-
-				// Remove public and standalone flag so garbage collection can delete the object.
-				ObjectToDelete->ClearFlags(RF_Standalone | RF_Public);
-
-				// Mark its package as dirty as we're going to delete it.
-				ObjectToDelete->MarkPendingKill();
-				ObjectToDelete->MarkPackageDirty();
-
-				// Notify the asset registry
-				FAssetRegistryModule::AssetDeleted( ObjectToDelete );
+				ObjectToDelete->Rename( nullptr, GetTransientPackage(), REN_DontCreateRedirectors | REN_NonTransactional );
+				ObjectsToDelete.Add( ObjectToDelete );
 			}
 		}
 	}
 
+	// Prevent display of progress dialog from ObjectTools::DeleteObjectsUnchecked
+	bool bGIsSilentCached = GIsSilent;
+	GIsSilent = true;
+
+	ensure( ObjectTools::DeleteObjectsUnchecked( ObjectsToDelete ) == ObjectsToDelete.Num() );
+
+	// Restore GIsSilent
+	GIsSilent = bGIsSilentCached;
+
 	CachedAssets.Reset();
 	Assets.Reset();
 
@@ -898,7 +900,7 @@ TSharedRef<SDockTab> FDataprepEditor::SpawnTabScenePreview(const FSpawnTabArgs &
 
 	return SNew(SDockTab)
 		//.Icon(FDataprepEditorStyle::Get()->GetBrush("DataprepEditor.Tabs.ScenePreview"))
-		.Label(LOCTEXT("DataprepEditor_ScenePreviewTab_Title", "ScenePreview"))
+		.Label(LOCTEXT("DataprepEditor_ScenePreviewTab_Title", "Scene Preview"))
 		[
 			ScenePreviewView.ToSharedRef()
 		];
@@ -910,7 +912,7 @@ TSharedRef<SDockTab> FDataprepEditor::SpawnTabAssetPreview(const FSpawnTabArgs &
 
 	return SNew(SDockTab)
 		//.Icon(FDataprepEditorStyle::Get()->GetBrush("DataprepEditor.Tabs.AssetPreview"))
-		.Label(LOCTEXT("DataprepEditor_AssetPreviewTab_Title", "AssetPreview"))
+		.Label(LOCTEXT("DataprepEditor_AssetPreviewTab_Title", "Asset Preview"))
 		[
 			SNew(SBorder)
 			.Padding(2.f)
diff --git a/Engine/Plugins/Experimental/DataPrepEditor/Source/DataPrepEditor/Private/DataPrepEditorActions.cpp b/Engine/Plugins/Experimental/DataPrepEditor/Source/DataPrepEditor/Private/DataPrepEditorActions.cpp
--- a/Engine/Plugins/Experimental/DataPrepEditor/Source/DataPrepEditor/Private/DataPrepEditorActions.cpp
+++ b/Engine/Plugins/Experimental/DataPrepEditor/Source/DataPrepEditor/Private/DataPrepEditorActions.cpp
@@ -11,7 +11,7 @@ void FDataprepEditorCommands::RegisterCommands()
 	// end of temp code for nodes development
 
 	UI_COMMAND(SaveScene, "Save Datasmith Scene", "To be filled.", EUserInterfaceActionType::ToggleButton, FInputChord());
-	UI_COMMAND(ShowDataprepSettings, "Dataprep Settings", "Edit the data prep settings.", EUserInterfaceActionType::ToggleButton, FInputChord());
+	UI_COMMAND(ShowDataprepSettings, "Settings", "Edit the data prep settings.", EUserInterfaceActionType::ToggleButton, FInputChord());
 	UI_COMMAND(ShowDatasmithSceneSettings, "Scene Settings", "Edit the scene settings.", EUserInterfaceActionType::ToggleButton, FInputChord());
 	UI_COMMAND(BuildWorld, "Import", "Load external data sources referenced by producers into the Dataprep Editor", EUserInterfaceActionType::ToggleButton, FInputChord());
 	UI_COMMAND(ExecutePipeline, "Execute", "Execute the Dataprep graph on the imported data", EUserInterfaceActionType::ToggleButton, FInputChord());
diff --git a/Engine/Plugins/Experimental/DataPrepEditor/Source/DataPrepEditor/Private/DataprepLevelWriter.cpp b/Engine/Plugins/Experimental/DataPrepEditor/Source/DataPrepEditor/Private/DataprepLevelWriter.cpp
deleted file mode 100644
diff --git a/Engine/Plugins/Experimental/DataPrepEditor/Source/DataPrepEditor/Private/DataprepLevelWriter.h b/Engine/Plugins/Experimental/DataPrepEditor/Source/DataPrepEditor/Private/DataprepLevelWriter.h
deleted file mode 100644
diff --git a/Engine/Plugins/Experimental/DataPrepEditor/Source/DataPrepEditor/Private/Widgets/DataprepWidgets.cpp b/Engine/Plugins/Experimental/DataPrepEditor/Source/DataPrepEditor/Private/Widgets/DataprepWidgets.cpp
--- a/Engine/Plugins/Experimental/DataPrepEditor/Source/DataPrepEditor/Private/Widgets/DataprepWidgets.cpp
+++ b/Engine/Plugins/Experimental/DataPrepEditor/Source/DataPrepEditor/Private/Widgets/DataprepWidgets.cpp
@@ -59,7 +59,7 @@ void SDataprepConsumerWidget::OnBrowseContentFolder()
 		if ( PickContentPathDlg->ShowModal() == EAppReturnType::Ok )
 		{
 			DataprepConsumer->SetTargetContentFolder( PickContentPathDlg->GetPath().ToString() );
-			PackageTextBox->SetText( FText::FromString( DataprepConsumer->GetTargetContentFolder() ) );
+			UpdateContentFolderText();
 		}
 	}
 }
@@ -73,7 +73,7 @@ void SDataprepConsumerWidget::SetDataprepConsumer(UDataprepContentConsumer* InDa
 
 	DataprepConsumer = InDataprepConsumer;
 
-	PackageTextBox->SetText( FText::FromString( DataprepConsumer->GetTargetContentFolder() ) );
+	UpdateContentFolderText();
 	LevelTextBox->SetText( FText::FromString( DataprepConsumer->GetLevelName() ) );
 }
 
@@ -118,7 +118,7 @@ void SDataprepConsumerWidget::Construct(const FArguments& InArgs)
 						.VAlign(VAlign_Center)
 						.FillWidth(1.0f)
 						[
-							SAssignNew(PackageTextBox, SEditableTextBox)
+							SAssignNew(ContentFolderTextBox, SEditableTextBox)
 							.Text( FText::FromString( DataprepConsumer->GetTargetContentFolder() ) )
 							.HintText( LOCTEXT("DataprepSlateHelper_ContentFolderHintText", "Set the content folder to save in") )
 							.IsReadOnly( true )
@@ -164,6 +164,20 @@ void SDataprepConsumerWidget::Construct(const FArguments& InArgs)
 	];
 }
 
+void SDataprepConsumerWidget::UpdateContentFolderText()
+{
+	if(UDataprepContentConsumer* Consumer = DataprepConsumer.Get())
+	{
+		FString TargetContentFolder( Consumer->GetTargetContentFolder() );
+		if( TargetContentFolder.StartsWith( TEXT( "/Game" ) ) )
+		{
+			TargetContentFolder = TargetContentFolder.Replace( TEXT( "/Game" ), TEXT( "/Content" ) );
+		}
+
+		ContentFolderTextBox->SetText( FText::FromString( TargetContentFolder + TEXT( "/" ) ) );
+	}
+}
+
 TSharedRef< SWidget > SDataprepDetailsView::CreateDefaultWidget( TSharedPtr< SWidget >& NameWidget, TSharedPtr< SWidget >& ValueWidget, float LeftPadding, EHorizontalAlignment HAlign, EVerticalAlignment VAlign )
 {
 	/** Helper class to force a widget to fill in a space. Copied from SDetailSingleItemRow.cpp */
diff --git a/Engine/Plugins/Experimental/DataPrepEditor/Source/DataPrepEditor/Private/Widgets/DataprepWidgets.h b/Engine/Plugins/Experimental/DataPrepEditor/Source/DataPrepEditor/Private/Widgets/DataprepWidgets.h
--- a/Engine/Plugins/Experimental/DataPrepEditor/Source/DataPrepEditor/Private/Widgets/DataprepWidgets.h
+++ b/Engine/Plugins/Experimental/DataPrepEditor/Source/DataPrepEditor/Private/Widgets/DataprepWidgets.h
@@ -28,10 +28,11 @@ public:
 private:
 	void OnLevelNameChanged( const FText& NewLevelName, ETextCommit::Type CommitType );
 	void OnBrowseContentFolder();
+	void UpdateContentFolderText();
 
 private:
-	UDataprepContentConsumer* DataprepConsumer;
-	TSharedPtr< SEditableTextBox > PackageTextBox;
+	TWeakObjectPtr<UDataprepContentConsumer> DataprepConsumer;
+	TSharedPtr< SEditableTextBox > ContentFolderTextBox;
 	TSharedPtr< SEditableTextBox > LevelTextBox;
 };
 
