Copying CL #6980965 and #6896090 from //UE4/Dev-Enterprise/...

Fixed logic about superseding producers
Fixed clang compilation errors
Fixed orientation issue when loading CAD files

diff --git a/Engine/Plugins/Experimental/DataPrepEditor/Source/DataPrepCore/Private/DataPrepAsset.cpp b/Engine/Plugins/Experimental/DataPrepEditor/Source/DataPrepCore/Private/DataPrepAsset.cpp
--- a/Engine/Plugins/Experimental/DataPrepEditor/Source/DataPrepCore/Private/DataPrepAsset.cpp
+++ b/Engine/Plugins/Experimental/DataPrepEditor/Source/DataPrepCore/Private/DataPrepAsset.cpp
@@ -302,7 +302,10 @@ void UDataprepAsset::RunProducers(const UDataprepContentProducer::ProducerContex
 
 	for ( FDataprepAssetProducer& AssetProducer : Producers )
 	{
-		if (AssetProducer.Producer && ( AssetProducer.bIsEnabled && AssetProducer.SupersededBy == INDEX_NONE ) )
+		const bool bIsOkToRun = AssetProducer.Producer && AssetProducer.bIsEnabled &&
+			( AssetProducer.SupersededBy == INDEX_NONE || !Producers[AssetProducer.SupersededBy].bIsEnabled );
+
+		if ( bIsOkToRun )
 		{
 			FString OutReason;
 			if (AssetProducer.Producer->Initialize( InContext, OutReason ))
@@ -516,13 +519,13 @@ void UDataprepAsset::OnBlueprintChanged( UBlueprint* InBlueprint )
 	}
 }
 
-void UDataprepAsset::ValidateProducerChanges( int32 Index, bool &bChangeAll )
+void UDataprepAsset::ValidateProducerChanges( int32 InIndex, bool &bChangeAll )
 {
 	bChangeAll = false;
 
-	if( Producers.IsValidIndex( Index ) && Producers.Num() > 1 )
+	if( Producers.IsValidIndex( InIndex ) && Producers.Num() > 1 )
 	{
-		FDataprepAssetProducer& InAssetProducer = Producers[Index];
+		FDataprepAssetProducer& InAssetProducer = Producers[InIndex];
 
 		// Check if input producer is still superseded if applicable
 		if( InAssetProducer.SupersededBy != INDEX_NONE )
@@ -539,7 +542,10 @@ void UDataprepAsset::ValidateProducerChanges( int32 Index, bool &bChangeAll )
 		int32 SupersederIndex = 0;
 		for( FDataprepAssetProducer& AssetProducer : Producers )
 		{
-			if( AssetProducer.Producer != InAssetProducer.Producer && AssetProducer.Producer->Supersede( InAssetProducer.Producer ) )
+			if( AssetProducer.Producer != InAssetProducer.Producer &&
+				AssetProducer.bIsEnabled == true && 
+				AssetProducer.SupersededBy == INDEX_NONE && 
+				AssetProducer.Producer->Supersede( InAssetProducer.Producer ) )
 			{
 				// Disable found producer if another producer supersedes its production
 				InAssetProducer.SupersededBy = SupersederIndex;
@@ -554,7 +560,7 @@ void UDataprepAsset::ValidateProducerChanges( int32 Index, bool &bChangeAll )
 		{
 			if( AssetProducer.Producer != InAssetProducer.Producer )
 			{
-				if( AssetProducer.SupersededBy == Index )
+				if( AssetProducer.SupersededBy == InIndex )
 				{
 					if( !InAssetProducer.Producer->Supersede( AssetProducer.Producer ) )
 					{
@@ -562,10 +568,10 @@ void UDataprepAsset::ValidateProducerChanges( int32 Index, bool &bChangeAll )
 						AssetProducer.SupersededBy = INDEX_NONE;
 					}
 				}
-				else if( InAssetProducer.Producer->Supersede( AssetProducer.Producer ) )
+				else if( InAssetProducer.SupersededBy == INDEX_NONE && InAssetProducer.Producer->Supersede( AssetProducer.Producer ) )
 				{
 					bChangeAll = true;
-					AssetProducer.SupersededBy = Index;
+					AssetProducer.SupersededBy = InIndex;
 				}
 			}
 		}
diff --git a/Engine/Plugins/Experimental/DataPrepEditor/Source/DataPrepCore/Public/DataPrepAsset.h b/Engine/Plugins/Experimental/DataPrepEditor/Source/DataPrepCore/Public/DataPrepAsset.h
--- a/Engine/Plugins/Experimental/DataPrepEditor/Source/DataPrepCore/Public/DataPrepAsset.h
+++ b/Engine/Plugins/Experimental/DataPrepEditor/Source/DataPrepCore/Public/DataPrepAsset.h
@@ -280,16 +280,16 @@ public:
 		return Producers.IsValidIndex(Index) ? Producers[Index].Producer : nullptr;
 	}
 
-	/** @return pointer on producer at Index-th position in Producers array. nullptr if Index is invalid */
+	/** @return True if producer at Index-th position is enabled. Returns false if disabled or Index is invalid */
 	bool IsProducerEnabled(int32 Index) const
 	{ 
-		return Producers.IsValidIndex(Index) ? Producers[Index].bIsEnabled : nullptr;
+		return Producers.IsValidIndex(Index) ? Producers[Index].bIsEnabled : false;
 	}
 
-	/** @return pointer on producer at Index-th position in Producers array. nullptr if Index is invalid */
+	/** @return True if producer at Index-th position is superseded by an enabled producer. Returns false if not superseded or superseder is disabled or Index is invalid */
 	bool IsProducerSuperseded(int32 Index) const
 	{ 
-		return Producers.IsValidIndex(Index) ? Producers[Index].SupersededBy != INDEX_NONE : false;
+		return Producers.IsValidIndex(Index) ? Producers[Index].SupersededBy != INDEX_NONE && Producers.IsValidIndex(Producers[Index].SupersededBy) && Producers[Producers[Index].SupersededBy].bIsEnabled : false;
 	}
 
 	/**
diff --git a/Engine/Plugins/Experimental/DataPrepEditor/Source/DataPrepEditor/Private/BlueprintNodes/K2Node_DataprepProducer.cpp b/Engine/Plugins/Experimental/DataPrepEditor/Source/DataPrepEditor/Private/BlueprintNodes/K2Node_DataprepProducer.cpp
--- a/Engine/Plugins/Experimental/DataPrepEditor/Source/DataPrepEditor/Private/BlueprintNodes/K2Node_DataprepProducer.cpp
+++ b/Engine/Plugins/Experimental/DataPrepEditor/Source/DataPrepEditor/Private/BlueprintNodes/K2Node_DataprepProducer.cpp
@@ -353,7 +353,6 @@ TSharedPtr<SGraphNode> UK2Node_DataprepProducer::CreateVisualWidget()
 				MainBox->AddSlot()[
 					CreateProducerWidget( &ProducerWidgets[ Index ], Index )
 				];
-				++Index;
 			}
 		}
 
diff --git a/Engine/Plugins/Experimental/DataPrepEditor/Source/DataPrepEditor/Private/DataPrepEditor.cpp b/Engine/Plugins/Experimental/DataPrepEditor/Source/DataPrepEditor/Private/DataPrepEditor.cpp
--- a/Engine/Plugins/Experimental/DataPrepEditor/Source/DataPrepEditor/Private/DataPrepEditor.cpp
+++ b/Engine/Plugins/Experimental/DataPrepEditor/Source/DataPrepEditor/Private/DataPrepEditor.cpp
@@ -22,8 +22,8 @@
 #include "BlueprintNodeSpawner.h"
 #include "DesktopPlatformModule.h"
 #include "Dialogs/DlgPickPath.h"
-#include "EditorDirectories.h"
 #include "Editor.h"
+#include "EditorDirectories.h"
 #include "EditorStyleSet.h"
 #include "Engine/World.h"
 #include "EngineUtils.h"
@@ -42,6 +42,7 @@
 #include "ObjectTools.h"
 #include "SceneOutlinerModule.h"
 #include "ScopedTransaction.h"
+#include "Templates/UnrealTemplate.h"
 #include "Toolkits/IToolkit.h"
 #include "UnrealEdGlobals.h"
 #include "Widgets/Docking/SDockTab.h"
@@ -605,22 +606,15 @@ void FDataprepEditor::CleanPreviewWorld()
 	}
 
 	// Prevent display of progress dialog from ObjectTools::DeleteObjectsUnchecked
-	bool bGIsSilentCached = GIsSilent;
-	GIsSilent = true;
+	TGuardValue<bool>( GIsSilent, true );
 
 	ensure( ObjectTools::DeleteObjectsUnchecked( ObjectsToDelete ) == ObjectsToDelete.Num() );
 
-	// Restore GIsSilent
-	GIsSilent = bGIsSilentCached;
-
 	CachedAssets.Reset();
 	Assets.Reset();
 
-	if (GEngine)
-	{
-		// Otherwise we end up with some issues
-		GEngine->PerformGarbageCollectionAndCleanupActors();
-	}
+	CollectGarbage( GARBAGE_COLLECTION_KEEPFLAGS );
+	PreviewWorld->CleanupActors();
 }
 
 void FDataprepEditor::OnExecutePipeline()
