Copying CL #7048349 from //UE4/Dev-Enterprise/...

Fixed issue with broken link between level sequence and referred actor
Changed syntax in list of extensions from "abc;def;ghi" to "*.abc;*.def;*.ghi"

diff --git a/Engine/Plugins/Experimental/DataPrepEditor/Source/DataPrepCore/DataPrepCore.Build.cs b/Engine/Plugins/Experimental/DataPrepEditor/Source/DataPrepCore/DataPrepCore.Build.cs
--- a/Engine/Plugins/Experimental/DataPrepEditor/Source/DataPrepCore/DataPrepCore.Build.cs
+++ b/Engine/Plugins/Experimental/DataPrepEditor/Source/DataPrepCore/DataPrepCore.Build.cs
@@ -18,6 +18,7 @@ namespace UnrealBuildTool.Rules
 					"Engine",
 					"GraphEditor",
 					"InputCore",
+					"LevelSequence",
 					"KismetCompiler",
 					"MeshDescription",
 					"MeshDescriptionOperations",
diff --git a/Engine/Plugins/Experimental/DataPrepEditor/Source/DataPrepCore/Private/DataPrepContentProducer.cpp b/Engine/Plugins/Experimental/DataPrepEditor/Source/DataPrepCore/Private/DataPrepContentProducer.cpp
--- a/Engine/Plugins/Experimental/DataPrepEditor/Source/DataPrepCore/Private/DataPrepContentProducer.cpp
+++ b/Engine/Plugins/Experimental/DataPrepEditor/Source/DataPrepCore/Private/DataPrepContentProducer.cpp
@@ -2,9 +2,11 @@
 
 #include "DataPrepContentProducer.h"
 
+#include "AssetToolsModule.h"
 #include "Engine/World.h"
 #include "EngineUtils.h"
 #include "GameFramework/Actor.h"
+#include "LevelSequence.h"
 
 const FString DefaultNamespace( TEXT("void") );
 
@@ -45,10 +47,27 @@ bool UDataprepContentProducer::Produce()
 		}
 	}
 
+	// Cache number of assets to go through new assets after call to Execute
+	int32 LastAssetCount = Assets.Num();
+
 	// Prefix all newly created actors with the namespace of the producer
 	// #ueent_todo: find a better way to identify newly created actors
 	if( Execute() )
 	{
+		// Collect all packages containing LevelSequence assets to remap their reference to an newly created actor
+		TSet< UPackage* > LevelSequencePackagesToCheck;
+		for(int32 Index = LastAssetCount; Index < Assets.Num(); ++Index)
+		{
+			if( ULevelSequence* LevelSequence = Cast<ULevelSequence>( Assets[Index].Get() ) )
+			{
+				LevelSequencePackagesToCheck.Add( LevelSequence->GetOutermost() );
+			}
+		}
+
+		// Map between old path and new path of newly created actors
+		TMap< FSoftObjectPath, FSoftObjectPath > ActorRedirectorMap;
+
+		// Rename actors
 		const FString Namespace = GetNamespace();
 
 		for (TActorIterator<AActor> It(Context.WorldPtr.Get(), AActor::StaticClass(), Flags); It; ++It)
@@ -57,11 +76,22 @@ bool UDataprepContentProducer::Produce()
 			{
 				AActor* Actor = *It;
 
+				FSoftObjectPath PreviousActorSoftPath(Actor);
+
 				const FString ActorName =  Namespace + TEXT("_") + Actor->GetName();
 				Actor->Rename( *ActorName, nullptr );
+
+				ActorRedirectorMap.Emplace( PreviousActorSoftPath, Actor );
 			}
 		}
 
+		// Update reference of LevelSequence assets if necessary
+		if(LevelSequencePackagesToCheck.Num() > 0)
+		{
+			IAssetTools& AssetTools = FModuleManager::LoadModuleChecked<FAssetToolsModule>("AssetTools").Get();
+			AssetTools.RenameReferencingSoftObjectPaths( LevelSequencePackagesToCheck.Array(), ActorRedirectorMap );
+		}
+
 		return true;
 	}
 
