Merging 6956441 from //UE4/Dev-Enterprise/... to //UE4/Release-4.23/...

Fixed the issue were we couldn't delete a dataprep asset when the editor is open.
The issue was that the embeded blueprint was flag standalone and refered the asset.

Also removed traces of the static mesh editor and some old useless code from the dataprep editor.

diff --git a/Engine/Plugins/Experimental/DataPrepEditor/Source/DataPrepCore/Private/DataPrepAsset.cpp b/Engine/Plugins/Experimental/DataPrepEditor/Source/DataPrepCore/Private/DataPrepAsset.cpp
--- a/Engine/Plugins/Experimental/DataPrepEditor/Source/DataPrepCore/Private/DataPrepAsset.cpp
+++ b/Engine/Plugins/Experimental/DataPrepEditor/Source/DataPrepCore/Private/DataPrepAsset.cpp
@@ -159,6 +159,9 @@ void UDataprepAsset::PostInitProperties()
 		DataprepRecipeBP = FKismetEditorUtilities::CreateBlueprint( UDataprepRecipe::StaticClass(), this, BlueprintName, BPTYPE_Normal, UBlueprint::StaticClass(), UBlueprintGeneratedClass::StaticClass() );
 		check( DataprepRecipeBP );
 
+		// This blueprint is not the asset of the package
+		DataprepRecipeBP->ClearFlags( RF_Standalone );
+
 		FAssetRegistryModule::AssetCreated( DataprepRecipeBP );
 
 		DataprepRecipeBP->MarkPackageDirty();
diff --git a/Engine/Plugins/Experimental/DataPrepEditor/Source/DataPrepEditor/Private/DataPrepEditor.cpp b/Engine/Plugins/Experimental/DataPrepEditor/Source/DataPrepEditor/Private/DataPrepEditor.cpp
--- a/Engine/Plugins/Experimental/DataPrepEditor/Source/DataPrepEditor/Private/DataPrepEditor.cpp
+++ b/Engine/Plugins/Experimental/DataPrepEditor/Source/DataPrepEditor/Private/DataPrepEditor.cpp
@@ -212,12 +212,12 @@ FName FDataprepEditor::GetToolkitFName() const
 
 FText FDataprepEditor::GetBaseToolkitName() const
 {
-	return LOCTEXT("AppLabel", "StaticMesh Editor");
+	return LOCTEXT("AppLabel", "Dataprep Editor");
 }
 
 FString FDataprepEditor::GetWorldCentricTabPrefix() const
 {
-	return LOCTEXT("WorldCentricTabPrefix", "StaticMesh ").ToString();
+	return LOCTEXT("WorldCentricTabPrefix", "Dataprep").ToString();
 }
 
 FLinearColor FDataprepEditor::GetWorldCentricTabColorScale() const
@@ -294,8 +294,6 @@ void FDataprepEditor::InitDataprepEditor(const EToolkitMode::Type Mode, const TS
 	UEdGraph* PipelineGraph = FBlueprintEditorUtils::FindEventGraph(DataprepRecipeBPPtr.Get());
 	check( PipelineGraph );
 
-	TArray< class UK2Node_DataprepAction* > ActionNodes;
-
 	for( UEdGraphNode* GraphNode : PipelineGraph->Nodes )
 	{
 		if( GraphNode->IsA<UK2Node_DataprepProducer>() )
@@ -308,8 +306,6 @@ void FDataprepEditor::InitDataprepEditor(const EToolkitMode::Type Mode, const TS
 		}
 	}
 
-	FDataprepEditorStyle::Initialize();
-
 	// This should normally happen only with a brand new Dataprep asset
 	if( StartNode == nullptr )
 	{
@@ -399,7 +395,7 @@ void FDataprepEditor::InitDataprepEditor(const EToolkitMode::Type Mode, const TS
 
 	const bool bCreateDefaultStandaloneMenu = true;
 	const bool bCreateDefaultToolbar = true;
-	FAssetEditorToolkit::InitAssetEditor(Mode, InitToolkitHost, DataprepEditorAppIdentifier, StandaloneDefaultLayout, bCreateDefaultToolbar, bCreateDefaultStandaloneMenu, DataprepAssetPtr.Get());
+	FAssetEditorToolkit::InitAssetEditor( Mode, InitToolkitHost, DataprepEditorAppIdentifier, StandaloneDefaultLayout, bCreateDefaultToolbar, bCreateDefaultStandaloneMenu, InDataprepAsset );
 
 	ExtendMenu();
 	ExtendToolBar();
@@ -474,13 +470,6 @@ void FDataprepEditor::OnBuildWorld()
 		return;
 	}
 
-	// Temp code for the nodes development
-	UDataprepRecipe* DataprepRecipePtr = CastChecked<UDataprepRecipe>(DataprepRecipeBPPtr->GeneratedClass->GetDefaultObject());
-
-	DataprepRecipePtr->SetTargetWorld(nullptr);
-	DataprepRecipePtr->ResetAssets();
-	// end of temp code for nodes development
-
 	if (DataprepAsset->GetProducersCount() == 0)
 	{
 		ResetBuildWorld();
@@ -629,17 +618,6 @@ void FDataprepEditor::OnExecutePipeline()
 		RestoreFromSnapshot();
 	}
 
-	// Temp code for the nodes development
-	UDataprepRecipe* DataprepRecipePtr = CastChecked<UDataprepRecipe>( DataprepRecipeBPPtr->GeneratedClass->GetDefaultObject() );
-
-	// Set the world to pull actors from when executing pipeline
-	DataprepRecipePtr->SetTargetWorld(PreviewWorld.Get());
-
-	// Set array of assets to be processed when executing pipeline
-	DataprepRecipePtr->SetAssets(Assets);
-
-	// end of temp code for nodes development
-
 	TSharedPtr< IDataprepProgressReporter > ProgressReporter( new FDataprepProgressReporter( LOCTEXT("Dataprep_ExecutePipeline", "Executing pipeline ...") ) );
 
 	// Trigger execution of data preparation operations on world attached to recipe
diff --git a/Engine/Plugins/Experimental/DataPrepEditor/Source/DataPrepEditor/Private/DataPrepEditor.h b/Engine/Plugins/Experimental/DataPrepEditor/Source/DataPrepEditor/Private/DataPrepEditor.h
--- a/Engine/Plugins/Experimental/DataPrepEditor/Source/DataPrepEditor/Private/DataPrepEditor.h
+++ b/Engine/Plugins/Experimental/DataPrepEditor/Source/DataPrepEditor/Private/DataPrepEditor.h
@@ -24,7 +24,6 @@ class SDockableTab;
 class SGraphNodeDetailsWidget;
 class SInspectorView;
 class SWidget;
-class UDataprepRecipe;
 class UWorld;
 class UEdGraphNode;
 
@@ -76,12 +75,6 @@ public:
 	virtual FText GetBaseToolkitName() const override;
 	virtual FString GetWorldCentricTabPrefix() const override;
 
-	/** @return the documentation location for this editor */
-	virtual FString GetDocumentationLink() const override
-	{
-		return FString(TEXT("Engine/Content/Types/StaticMeshes/Editor"));
-	}
-
 	/** @return Returns the color and opacity to use for the color that appears behind the tab text for this toolkit's tab in world-centric mode. */
 	virtual FLinearColor GetWorldCentricTabColorScale() const override;
 
