Merging 7039746 and 7040354  from //UE4/Dev-Enterprise/... to //UE4/Release-4.23/...

Fixed some clipping issues.

diff --git a/Engine/Plugins/Experimental/DataPrepEditor/Source/DataPrepEditor/Private/Widgets/DataprepWidgets.cpp b/Engine/Plugins/Experimental/DataPrepEditor/Source/DataPrepEditor/Private/Widgets/DataprepWidgets.cpp
--- a/Engine/Plugins/Experimental/DataPrepEditor/Source/DataPrepEditor/Private/Widgets/DataprepWidgets.cpp
+++ b/Engine/Plugins/Experimental/DataPrepEditor/Source/DataPrepEditor/Private/Widgets/DataprepWidgets.cpp
@@ -222,6 +222,7 @@ TSharedRef< SWidget > SDataprepDetailsView::CreateDefaultWidget( TSharedPtr< SWi
 		.OnSlotResized( SSplitter::FOnSlotResized::CreateSP( this, &SDataprepDetailsView::OnLeftColumnResized ) )
 		[
 			SNew(SHorizontalBox)
+			.Clipping(EWidgetClipping::OnDemand)
 			+ SHorizontalBox::Slot()
 			.VAlign(VAlign_Center)
 			.HAlign(HAlign_Left)
@@ -235,6 +236,7 @@ TSharedRef< SWidget > SDataprepDetailsView::CreateDefaultWidget( TSharedPtr< SWi
 		.OnSlotResized(ColumnSizeData.OnWidthChanged)
 		[
 			SNew(SHorizontalBox)
+			.Clipping(EWidgetClipping::OnDemand)
 			+ SHorizontalBox::Slot()
 			.Padding( 5.0f, 2.5f, 2.0f, 2.5f )
 			.HAlign( HAlign )
