Fixed crash when loading dtatprep asset with incomplete consumer or producer

diff --git a/Engine/Plugins/Experimental/DataPrepEditor/DataPrepEditor.uplugin b/Engine/Plugins/Experimental/DataPrepEditor/DataPrepEditor.uplugin
--- a/Engine/Plugins/Experimental/DataPrepEditor/DataPrepEditor.uplugin
+++ b/Engine/Plugins/Experimental/DataPrepEditor/DataPrepEditor.uplugin
@@ -1,8 +1,8 @@
 {
 	"FileVersion": 3,
 	"Version": 1,
-	"VersionName": "0.1",
-	"FriendlyName": "Datasmith DataPrep Editor",
+	"VersionName": "0.9",
+	"FriendlyName": "DataPrep Editor",
 	"Description": "[WIP] Allows editing of Data preparation pipeline from within the Unreal editor",
 	"Category": "Programming",
 	"CreatedBy": "Epic Games, Inc.",
diff --git a/Engine/Plugins/Experimental/DataPrepEditor/Source/DataPrepCore/Private/DataPrepAsset.cpp b/Engine/Plugins/Experimental/DataPrepEditor/Source/DataPrepCore/Private/DataPrepAsset.cpp
--- a/Engine/Plugins/Experimental/DataPrepEditor/Source/DataPrepCore/Private/DataPrepAsset.cpp
+++ b/Engine/Plugins/Experimental/DataPrepEditor/Source/DataPrepCore/Private/DataPrepAsset.cpp
@@ -177,15 +177,20 @@ void UDataprepAsset::Serialize( FArchive& Ar )
 
 	if( Ar.IsLoading() )
 	{
-		check(Consumer);
-		Consumer->GetOnChanged().AddUObject( this, &UDataprepAsset::OnConsumerChanged );
+		if(Consumer != nullptr)
+		{
+			Consumer->GetOnChanged().AddUObject( this, &UDataprepAsset::OnConsumerChanged );
+		}
 
 		check(DataprepRecipeBP);
 		DataprepRecipeBP->OnChanged().AddUObject( this, &UDataprepAsset::OnBlueprintChanged );
 
 		for( FDataprepAssetProducer& Producer : Producers )
 		{
-			Producer.Producer->GetOnChanged().AddUObject( this, &UDataprepAsset::OnProducerChanged );
+			if(Producer.Producer)
+			{
+				Producer.Producer->GetOnChanged().AddUObject( this, &UDataprepAsset::OnProducerChanged );
+			}
 		}
 	}
 }
@@ -339,22 +344,27 @@ void UDataprepAsset::RunProducers(const UDataprepContentProducer::ProducerContex
 
 bool UDataprepAsset::RunConsumer( const UDataprepContentConsumer::ConsumerContext& InContext, FString& OutReason)
 {
-	if( !Consumer->Initialize( InContext, OutReason ) )
+	if(Consumer)
 	{
-		return false;
-	}
+		if( !Consumer->Initialize( InContext, OutReason ) )
+		{
+			return false;
+		}
 
-	// #ueent_todo: Update state of entry: finalizing
+		// #ueent_todo: Update state of entry: finalizing
 
-	if ( !Consumer->Run() )
-	{
-		// #ueent_todo: Inform execution has failed
-		return false;
+		if ( !Consumer->Run() )
+		{
+			// #ueent_todo: Inform execution has failed
+			return false;
+		}
+
+		Consumer->Reset();
+
+		return true;
 	}
 
-	Consumer->Reset();
-	
-	return true;
+	return false;
 }
 
 bool UDataprepAsset::AddProducer(UClass* ProducerClass)
@@ -383,11 +393,12 @@ bool UDataprepAsset::RemoveProducer(int32 IndexToRemove)
 {
 	if( Producers.IsValidIndex( IndexToRemove ) )
 	{
-		UDataprepContentProducer* Producer = Producers[IndexToRemove].Producer;
-
-		Producer->GetOnChanged().RemoveAll( this );
+		if(UDataprepContentProducer* Producer = Producers[IndexToRemove].Producer)
+		{
+			Producer->GetOnChanged().RemoveAll( this );
 
-		DataprepAssetUtil::DeleteRegisteredAsset( Producer );
+			DataprepAssetUtil::DeleteRegisteredAsset( Producer );
+		}
 
 		Producers.RemoveAt( IndexToRemove );
 
@@ -535,7 +546,7 @@ void UDataprepAsset::ValidateProducerChanges( int32 InIndex, bool &bChangeAll )
 		{
 			FDataprepAssetProducer& SupersedingAssetProducer = Producers[ InAssetProducer.SupersededBy ];
 
-			if( !SupersedingAssetProducer.Producer->Supersede( InAssetProducer.Producer ) )
+			if( SupersedingAssetProducer.Producer != nullptr && !SupersedingAssetProducer.Producer->Supersede( InAssetProducer.Producer ) )
 			{
 				InAssetProducer.SupersededBy = INDEX_NONE;
 			}
@@ -545,7 +556,8 @@ void UDataprepAsset::ValidateProducerChanges( int32 InIndex, bool &bChangeAll )
 		int32 SupersederIndex = 0;
 		for( FDataprepAssetProducer& AssetProducer : Producers )
 		{
-			if( AssetProducer.Producer != InAssetProducer.Producer &&
+			if( AssetProducer.Producer != nullptr &&
+				AssetProducer.Producer != InAssetProducer.Producer &&
 				AssetProducer.bIsEnabled == true && 
 				AssetProducer.SupersededBy == INDEX_NONE && 
 				AssetProducer.Producer->Supersede( InAssetProducer.Producer ) )
@@ -559,23 +571,26 @@ void UDataprepAsset::ValidateProducerChanges( int32 InIndex, bool &bChangeAll )
 
 		// If input producer superseded any other producer, check if this is still valid
 		// Check if input producer does not supersede other producers
-		for( FDataprepAssetProducer& AssetProducer : Producers )
+		if( InAssetProducer.Producer != nullptr )
 		{
-			if( AssetProducer.Producer != InAssetProducer.Producer )
+			for( FDataprepAssetProducer& AssetProducer : Producers )
 			{
-				if( AssetProducer.SupersededBy == InIndex )
+				if( AssetProducer.Producer != InAssetProducer.Producer )
 				{
-					if( !InAssetProducer.Producer->Supersede( AssetProducer.Producer ) )
+					if( AssetProducer.SupersededBy == InIndex )
+					{
+						if( !InAssetProducer.Producer->Supersede( AssetProducer.Producer ) )
+						{
+							bChangeAll = true;
+							AssetProducer.SupersededBy = INDEX_NONE;
+						}
+					}
+					else if( InAssetProducer.SupersededBy == INDEX_NONE && InAssetProducer.Producer->Supersede( AssetProducer.Producer ) )
 					{
 						bChangeAll = true;
-						AssetProducer.SupersededBy = INDEX_NONE;
+						AssetProducer.SupersededBy = InIndex;
 					}
 				}
-				else if( InAssetProducer.SupersededBy == INDEX_NONE && InAssetProducer.Producer->Supersede( AssetProducer.Producer ) )
-				{
-					bChangeAll = true;
-					AssetProducer.SupersededBy = InIndex;
-				}
 			}
 		}
 	}
diff --git a/Engine/Plugins/Experimental/DataPrepEditor/Source/DataPrepEditor/Private/DataPrepEditor.cpp b/Engine/Plugins/Experimental/DataPrepEditor/Source/DataPrepEditor/Private/DataPrepEditor.cpp
--- a/Engine/Plugins/Experimental/DataPrepEditor/Source/DataPrepEditor/Private/DataPrepEditor.cpp
+++ b/Engine/Plugins/Experimental/DataPrepEditor/Source/DataPrepEditor/Private/DataPrepEditor.cpp
@@ -913,7 +913,7 @@ void FDataprepEditor::UpdatePreviewPanels()
 	//				   that new assets have been generated.
 	AssetPreviewView->ClearAssetList();
 	FString SubstitutePath = DataprepAssetPtr->GetOutermost()->GetName();
-	if(!DataprepAssetPtr->GetConsumer()->GetTargetContentFolder().IsEmpty())
+	if(DataprepAssetPtr->GetConsumer() != nullptr && !DataprepAssetPtr->GetConsumer()->GetTargetContentFolder().IsEmpty())
 	{
 		SubstitutePath = DataprepAssetPtr->GetConsumer()->GetTargetContentFolder();
 	}
diff --git a/Engine/Plugins/Experimental/DataPrepEditor/Source/DataPrepEditor/Private/Widgets/Action/SDataprepFilter.cpp b/Engine/Plugins/Experimental/DataPrepEditor/Source/DataPrepEditor/Private/Widgets/Action/SDataprepFilter.cpp
--- a/Engine/Plugins/Experimental/DataPrepEditor/Source/DataPrepEditor/Private/Widgets/Action/SDataprepFilter.cpp
+++ b/Engine/Plugins/Experimental/DataPrepEditor/Source/DataPrepEditor/Private/Widgets/Action/SDataprepFilter.cpp
@@ -128,6 +128,7 @@ TSharedRef<SWidget> SDataprepFilter::GetContentWidget() const
 				{
 					return Filter->GetFetcher();
 				})
+			.Class( UDataprepFilter::StaticClass() )
 		];
 }
 
diff --git a/Engine/Plugins/Experimental/DataPrepEditor/Source/DataPrepEditor/Private/Widgets/Action/SDataprepOperation.cpp b/Engine/Plugins/Experimental/DataPrepEditor/Source/DataPrepEditor/Private/Widgets/Action/SDataprepOperation.cpp
--- a/Engine/Plugins/Experimental/DataPrepEditor/Source/DataPrepEditor/Private/Widgets/Action/SDataprepOperation.cpp
+++ b/Engine/Plugins/Experimental/DataPrepEditor/Source/DataPrepEditor/Private/Widgets/Action/SDataprepOperation.cpp
@@ -23,12 +23,7 @@ FText SDataprepOperation::GetBlockTitle() const
 
 TSharedRef<SWidget> SDataprepOperation::GetContentWidget() const
 {
-	if ( Operation )
-	{
-		return SNew( SDataprepDetailsView ).Object( Operation );
-	}
-
-	return SNullWidget::NullWidget;
+	return SNew( SDataprepDetailsView ).Object( Operation ).Class( UDataprepOperation::StaticClass() );
 }
 
 void SDataprepOperation::AddReferencedObjects(FReferenceCollector& Collector)
diff --git a/Engine/Plugins/Experimental/DataPrepEditor/Source/DataPrepEditor/Private/Widgets/DataprepAssetView.cpp b/Engine/Plugins/Experimental/DataPrepEditor/Source/DataPrepEditor/Private/Widgets/DataprepAssetView.cpp
--- a/Engine/Plugins/Experimental/DataPrepEditor/Source/DataPrepEditor/Private/Widgets/DataprepAssetView.cpp
+++ b/Engine/Plugins/Experimental/DataPrepEditor/Source/DataPrepEditor/Private/Widgets/DataprepAssetView.cpp
@@ -325,11 +325,8 @@ void SProducerStackEntryTreeView::BuildProducerEntries()
 
 		for( int32 Index = 0; Index < ProducersCount; ++Index )
 		{
-			if(DataprepAsset->GetProducer( Index ) != nullptr)
-			{
-				TSharedRef<FProducerStackEntry> ProducerStackEntry = MakeShareable( new FProducerStackEntry( Index, DataprepAsset ) );
-				RootNodes.Add( ProducerStackEntry );
-			}
+			TSharedRef<FProducerStackEntry> ProducerStackEntry = MakeShareable( new FProducerStackEntry( Index, DataprepAsset ) );
+			RootNodes.Add( ProducerStackEntry );
 		}
 	}
 }
@@ -367,7 +364,7 @@ TSharedRef<SWidget> SProducerStackEntryTableRow::GetInputMainWidget()
 {
 	FProducerStackEntryPtr ProducerStackEntry = Node.Pin();
 
-	if ( !ProducerStackEntry.IsValid() || !ProducerStackEntry->HasValidData() )
+	if ( !ProducerStackEntry.IsValid() )
 	{
 		return SNullWidget::NullWidget;
 	}
@@ -391,7 +388,6 @@ TSharedRef<SWidget> SProducerStackEntryTableRow::GetInputMainWidget()
 
 	TSharedPtr<SWidget> Widget = SNew(SBorder)
 	.BorderImage(FEditorStyle::GetBrush("NoBrush"))
-	.BorderBackgroundColor( ProducerStackEntry->bIsSuperseded ? FLinearColor::White : FLinearColor::Red )
 	.Padding(5.0f)
 	[
 		SNew(SHorizontalBox)
@@ -421,6 +417,7 @@ TSharedRef<SWidget> SProducerStackEntryTableRow::GetInputMainWidget()
 		[
 			SNew( SDataprepDetailsView )
 			.Object( ProducerStackEntry->GetProducer() )
+			.Class( UDataprepContentProducer::StaticClass() )
 		]
 		// Delete button
 		+ SHorizontalBox::Slot()
@@ -445,7 +442,7 @@ TSharedRef<SWidget> SProducerStackEntryTableRow::GetInputMainWidget()
 	];
 
 	// Disable check box if associated producer is superseeded
-	CheckBox->SetEnabled( !ProducerStackEntry->bIsSuperseded );
+	CheckBox->SetEnabled( ProducerStackEntry->WillBeRun() );
 
 	return Widget.ToSharedRef();
 }
diff --git a/Engine/Plugins/Experimental/DataPrepEditor/Source/DataPrepEditor/Private/Widgets/DataprepWidgets.cpp b/Engine/Plugins/Experimental/DataPrepEditor/Source/DataPrepEditor/Private/Widgets/DataprepWidgets.cpp
--- a/Engine/Plugins/Experimental/DataPrepEditor/Source/DataPrepEditor/Private/Widgets/DataprepWidgets.cpp
+++ b/Engine/Plugins/Experimental/DataPrepEditor/Source/DataPrepEditor/Private/Widgets/DataprepWidgets.cpp
@@ -53,7 +53,7 @@ void SDataprepConsumerWidget::OnBrowseContentFolder()
 	{
 		TSharedRef<SDlgPickPath> PickContentPathDlg =
 			SNew(SDlgPickPath)
-			.Title(LOCTEXT("ChooseImportRootContentPath", "Choose Location for importing the Datasmith content"))
+			.Title(LOCTEXT("DataprepSlateHelper_ChooseImportRootContentPath", "Choose Location for importing the Datasmith content"))
 			.DefaultPath(FText::FromString(Path));
 
 		if ( PickContentPathDlg->ShowModal() == EAppReturnType::Ok )
@@ -81,89 +81,124 @@ void SDataprepConsumerWidget::Construct(const FArguments& InArgs)
 {
 	DataprepConsumer = InArgs._DataprepConsumer;
 
+	ChildSlot
+	[
+		BuildWidget()
+	];
+}
+
+TSharedRef<SWidget> SDataprepConsumerWidget::BuildWidget()
+{
+	if(DataprepConsumer == nullptr)
+	{
+		return BuildNullWidget();
+	}
+
 	TSharedRef<SWidget> BrowseButton = PropertyCustomizationHelpers::MakeBrowseButton( FSimpleDelegate::CreateSP( this, &SDataprepConsumerWidget::OnBrowseContentFolder ) );
 
-	ChildSlot
+	return SNew(SBorder)
+	.BorderImage( FCoreStyle::Get().GetBrush("ToolPanel.GroupBorder") )
 	[
-		SNew(SBorder)
-		.BorderImage( FCoreStyle::Get().GetBrush("ToolPanel.GroupBorder") )
+		SNew(SVerticalBox)
+		+ SVerticalBox::Slot()
+		.Padding(5.0f)
 		[
-			SNew(SVerticalBox)
-			+ SVerticalBox::Slot()
-			.Padding(5.0f)
+			SNew(SGridPanel)
+			.FillColumn(0, 0.25f)
+			.FillColumn(1, 0.75f)
+			+ SGridPanel::Slot(0, 0)
 			[
-				SNew(SGridPanel)
-				.FillColumn(0, 0.25f)
-				.FillColumn(1, 0.75f)
-				+ SGridPanel::Slot(0, 0)
+				SNew(SHorizontalBox)
+				+ SHorizontalBox::Slot()
+				.HAlign(HAlign_Left)
+				.VAlign(VAlign_Center)
+				.FillWidth(1.0f)
 				[
-					SNew(SHorizontalBox)
-					+ SHorizontalBox::Slot()
-					.HAlign(HAlign_Left)
-					.VAlign(VAlign_Center)
-					.FillWidth(1.0f)
-					[
-						SNew(STextBlock)
-						.Text(LOCTEXT("DataprepSlateHelper_ContentFolderLabel", "Content Folder"))
-						.Margin(FMargin( 5.0f, 5.0f, 0.0f, 0.0f ) )
-					]
-				]
-				+ SGridPanel::Slot(1, 0)
-				.Padding(0.0f, 5.0f)
-				[
-					SNew(SBox)
-					[
-						SNew(SHorizontalBox)
-						+ SHorizontalBox::Slot()
-						.VAlign(VAlign_Center)
-						.FillWidth(1.0f)
-						[
-							SAssignNew(ContentFolderTextBox, SEditableTextBox)
-							.Text( FText::FromString( DataprepConsumer->GetTargetContentFolder() ) )
-							.HintText( LOCTEXT("DataprepSlateHelper_ContentFolderHintText", "Set the content folder to save in") )
-							.IsReadOnly( true )
-						]
-						+ SHorizontalBox::Slot()
-						.VAlign(VAlign_Center)
-						.HAlign(HAlign_Center)
-						[
-							BrowseButton
-						]
-					]			
+					SNew(STextBlock)
+					.Text(LOCTEXT("DataprepSlateHelper_ContentFolderLabel", "Content Folder"))
+					.Margin(FMargin( 5.0f, 5.0f, 0.0f, 0.0f ) )
 				]
-				+ SGridPanel::Slot(0, 1)
+			]
+			+ SGridPanel::Slot(1, 0)
+			.Padding(0.0f, 5.0f)
+			[
+				SNew(SBox)
 				[
 					SNew(SHorizontalBox)
 					+ SHorizontalBox::Slot()
-					.HAlign(HAlign_Left)
 					.VAlign(VAlign_Center)
 					.FillWidth(1.0f)
 					[
-						SNew(STextBlock)
-						.Text(LOCTEXT("DataprepSlateHelper_LevelNameLabel", "Level"))
-						.Margin(FMargin( 5.0f, 5.0f, 0.0f, 0.0f ) )
+						SAssignNew(ContentFolderTextBox, SEditableTextBox)
+						.Text( FText::FromString( DataprepConsumer->GetTargetContentFolder() ) )
+						.HintText( LOCTEXT("DataprepSlateHelper_ContentFolderHintText", "Set the content folder to save in") )
+						.IsReadOnly( true )
 					]
-				]
-				+ SGridPanel::Slot(1, 1)
-				.Padding(0.0f, 5.0f)
-				[
-					SNew(SHorizontalBox)
 					+ SHorizontalBox::Slot()
-					.HAlign(HAlign_Left)
 					.VAlign(VAlign_Center)
-					.FillWidth(1.0f)
+					.HAlign(HAlign_Center)
 					[
-						SAssignNew(LevelTextBox, SEditableTextBox)
-						.Text( FText::FromString( DataprepConsumer->GetLevelName() ) )
-						.HintText( LOCTEXT("DataprepSlateHelper_LevelNameHintText", "Current will be used") )
-						.OnTextCommitted(  this, &SDataprepConsumerWidget::OnLevelNameChanged  )
+						BrowseButton
 					]
+				]			
+			]
+			+ SGridPanel::Slot(0, 1)
+			[
+				SNew(SHorizontalBox)
+				+ SHorizontalBox::Slot()
+				.HAlign(HAlign_Left)
+				.VAlign(VAlign_Center)
+				.FillWidth(1.0f)
+				[
+					SNew(STextBlock)
+					.Text(LOCTEXT("DataprepSlateHelper_LevelNameLabel", "Level"))
+					.Margin(FMargin( 5.0f, 5.0f, 0.0f, 0.0f ) )
+				]
+			]
+			+ SGridPanel::Slot(1, 1)
+			.Padding(0.0f, 5.0f)
+			[
+				SNew(SHorizontalBox)
+				+ SHorizontalBox::Slot()
+				.HAlign(HAlign_Left)
+				.VAlign(VAlign_Center)
+				.FillWidth(1.0f)
+				[
+					SAssignNew(LevelTextBox, SEditableTextBox)
+					.Text( FText::FromString( DataprepConsumer->GetLevelName() ) )
+					.HintText( LOCTEXT("DataprepSlateHelper_LevelNameHintText", "Current will be used") )
+					.OnTextCommitted(  this, &SDataprepConsumerWidget::OnLevelNameChanged  )
 				]
 			]
 		]
 	];
 }
 
+TSharedRef<SWidget> SDataprepConsumerWidget::BuildNullWidget()
+{
+	return SNew(SBorder)
+	.BorderImage( FCoreStyle::Get().GetBrush("ToolPanel.GroupBorder") )
+	[
+		SNew(SVerticalBox)
+		+ SVerticalBox::Slot()
+		.Padding(5.0f)
+		[
+			SNew(SHorizontalBox)
+			+ SHorizontalBox::Slot()
+			.HAlign(HAlign_Left)
+			.VAlign(VAlign_Center)
+			.FillWidth(1.0f)
+			[
+				SNew(STextBlock)
+				.Font(FCoreStyle::GetDefaultFontStyle("Bold", 11))
+				.Text(LOCTEXT("DataprepSlateHelper_ContentFolderLabel", "Error: Not a valid consumer"))
+				.Margin(FMargin( 5.0f, 5.0f, 0.0f, 0.0f ) )
+				.ColorAndOpacity( FLinearColor(1,0,0,1) )
+			]
+		]
+	];
+}
+
 void SDataprepConsumerWidget::UpdateContentFolderText()
 {
 	if(UDataprepContentConsumer* Consumer = DataprepConsumer.Get())
@@ -404,9 +439,7 @@ void SDataprepDetailsView::AddWidgets( const TArray< TSharedRef< IDetailTreeNode
 
 void SDataprepDetailsView::Construct(const FArguments& InArgs)
 {
-	DetailedObject = InArgs._Object.Get();
 	ObjectAttribute =  InArgs._Object;
-
 	if( !ObjectAttribute.IsSet() )
 	{
 		ChildSlot
@@ -417,6 +450,9 @@ void SDataprepDetailsView::Construct(const FArguments& InArgs)
 		return;
 	}
 
+	DetailedObject = InArgs._Object.Get();
+	DetailedClass = InArgs._Class.Get();
+
 	FPropertyEditorModule& PropertyEditorModule = FModuleManager::GetModuleChecked<FPropertyEditorModule>("PropertyEditor");
 
 	FPropertyRowGeneratorArgs Args;
@@ -424,9 +460,12 @@ void SDataprepDetailsView::Construct(const FArguments& InArgs)
 
 	OnPropertyChangedHandle = Generator->OnFinishedChangingProperties().AddSP( this, &SDataprepDetailsView::OnPropertyChanged );
 
-	TArray< UObject* > Objects;
-	Objects.Add( DetailedObject );
-	Generator->SetObjects( Objects );
+	if( DetailedObject != nullptr )
+	{
+		TArray< UObject* > Objects;
+		Objects.Add( DetailedObject );
+		Generator->SetObjects( Objects );
+	}
 
 	ColumnWidth = 0.7f;
 	ColumnSizeData.LeftColumnWidth = TAttribute<float>( this, &SDataprepDetailsView::OnGetLeftColumnWidth );
@@ -459,10 +498,28 @@ void SDataprepDetailsView::Construct()
 	}
 	else
 	{
+		FText ErrorText = FText::Format( LOCTEXT( "DataprepSlateHelper_InvalidDetailedObject", "Error: Not a valid {0}" ), FText::FromString( DetailedClass ? DetailedClass->GetName() : UObject::StaticClass()->GetName() ) );
+
 		ChildSlot
+		[
+			SNew(SVerticalBox)
+			+ SVerticalBox::Slot()
+			.Padding(5.0f)
 			[
-				SNullWidget::NullWidget
-			];
+				SNew(SHorizontalBox)
+				+ SHorizontalBox::Slot()
+				.HAlign(HAlign_Left)
+				.VAlign(VAlign_Center)
+				.FillWidth(1.0f)
+				[
+					SNew(STextBlock)
+					.Font(FCoreStyle::GetDefaultFontStyle("Bold", 11))
+					.Text( ErrorText )
+					.Margin(FMargin( 5.0f, 5.0f, 0.0f, 0.0f ) )
+					.ColorAndOpacity( FLinearColor(1,0,0,1) )
+				]
+			]
+		];
 	}
 }
 
diff --git a/Engine/Plugins/Experimental/DataPrepEditor/Source/DataPrepEditor/Private/Widgets/DataprepWidgets.h b/Engine/Plugins/Experimental/DataPrepEditor/Source/DataPrepEditor/Private/Widgets/DataprepWidgets.h
--- a/Engine/Plugins/Experimental/DataPrepEditor/Source/DataPrepEditor/Private/Widgets/DataprepWidgets.h
+++ b/Engine/Plugins/Experimental/DataPrepEditor/Source/DataPrepEditor/Private/Widgets/DataprepWidgets.h
@@ -26,6 +26,8 @@ public:
 	void SetDataprepConsumer( UDataprepContentConsumer* DataprepConsumer );
 
 private:
+	TSharedRef<SWidget> BuildWidget();
+	TSharedRef<SWidget> BuildNullWidget();
 	void OnLevelNameChanged( const FText& NewLevelName, ETextCommit::Type CommitType );
 	void OnBrowseContentFolder();
 	void UpdateContentFolderText();
@@ -51,8 +53,11 @@ class SDataprepDetailsView : public SCompoundWidget
 
 public:
 	SLATE_BEGIN_ARGS(SDataprepDetailsView)
+		: _Object( nullptr )
+		, _Class( UObject::StaticClass() )
 	{}
 	SLATE_ATTRIBUTE(UObject*, Object)
+	SLATE_ATTRIBUTE(UClass*, Class)
 	SLATE_END_ARGS()
 
 public:
@@ -100,10 +105,13 @@ private:
 private:
 	/** Row generator applied on detailed object */
 	TSharedPtr< class IPropertyRowGenerator > Generator;
-	
+
 	/** Object to be detailed */
 	UObject* DetailedObject;
-	
+
+	/** Class of the object to be detailed */
+	UClass* DetailedClass;
+
 	TAttribute<UObject*> ObjectAttribute;
 
 	/** Array properties tracked for changes */
