Removed a lot of random crash due to the fact that pumping the messages of the slate application may lead to unwanted result while running the import, execute or commit of the dataprep.

diff --git a/Engine/Plugins/Experimental/DataPrepEditor/Source/DataPrepEditor/Private/DataPrepEditor.cpp b/Engine/Plugins/Experimental/DataPrepEditor/Source/DataPrepEditor/Private/DataPrepEditor.cpp
--- a/Engine/Plugins/Experimental/DataPrepEditor/Source/DataPrepEditor/Private/DataPrepEditor.cpp
+++ b/Engine/Plugins/Experimental/DataPrepEditor/Source/DataPrepEditor/Private/DataPrepEditor.cpp
@@ -481,14 +481,6 @@ void FDataprepEditor::OnBuildWorld()
 
 	CleanPreviewWorld();
 
-	if(UEdGraph* PipelineGraph = GetPipelineGraph())
-	{
-		PipelineView->ClearSelectionSet();
-		PipelineView->SetNodeSelection( StartNode, true );
-		FSlateApplication::Get().PumpMessages();
-		FSlateApplication::Get().Tick();
-	}
-
 	UPackage* TransientPackage = NewObject< UPackage >( nullptr, *GetTransientContentFolder(), RF_Transient );
 	TransientPackage->FullyLoad();
 
@@ -672,12 +664,6 @@ void FDataprepEditor::OnExecutePipeline()
 						}
 					}
 
-					// Highlight action being executed
-					PipelineView->ClearSelectionSet();
-					PipelineView->SetNodeSelection( ActionNode, true );
-					FSlateApplication::Get().PumpMessages();
-					FSlateApplication::Get().Tick();
-
 					// Execute action
 					ActionNode->GetDataprepAction()->Execute( Objects );
 					ActionNodesExecuted.Add( ActionNode );
@@ -699,8 +685,6 @@ void FDataprepEditor::OnExecutePipeline()
 
 					// World may have changed, update asset preview and scene outliner
 					UpdatePreviewPanels();
-					FSlateApplication::Get().PumpMessages();
-					FSlateApplication::Get().Tick();
 				}
 
 				UEdGraphPin* NextNodeOutPin = NextNode->FindPin( UEdGraphSchema_K2::PN_Then, EGPD_Output );
@@ -747,10 +731,6 @@ void FDataprepEditor::OnCommitWorld()
 		// #ueent_todo: Inform consumer failed
 	}
 
-	PipelineView->ClearSelectionSet();
-	FSlateApplication::Get().PumpMessages();
-	FSlateApplication::Get().Tick();
-
 	ResetBuildWorld();
 }
 
