Added dedicated panel for dataprep asset settings
Unified font used across widgets in Dataprep panel
Reworked layout of widgets in Dataprep panel
Fixed logic flaw in management of superseding
Fixed some cosmetic UI issues
Fixed an issue with duplicating a dataprep asset: wrong outer

diff --git a/Engine/Plugins/Experimental/DataPrepEditor/Source/DataPrepCore/Private/DataPrepAsset.cpp b/Engine/Plugins/Experimental/DataPrepEditor/Source/DataPrepCore/Private/DataPrepAsset.cpp
--- a/Engine/Plugins/Experimental/DataPrepEditor/Source/DataPrepCore/Private/DataPrepAsset.cpp
+++ b/Engine/Plugins/Experimental/DataPrepEditor/Source/DataPrepCore/Private/DataPrepAsset.cpp
@@ -141,7 +141,7 @@ void UDataprepAsset::PostInitProperties()
 			{
 				if( CurrentClass->IsChildOf( UDataprepContentConsumer::StaticClass() ) )
 				{
-					Consumer = NewObject< UDataprepContentConsumer >( GetOutermost(), CurrentClass, NAME_None, RF_Transactional );
+					Consumer = NewObject< UDataprepContentConsumer >( this, CurrentClass, NAME_None, RF_Transactional );
 					check( Consumer );
 
 					FAssetRegistryModule::AssetCreated( Consumer );
@@ -371,7 +371,7 @@ bool UDataprepAsset::AddProducer(UClass* ProducerClass)
 {
 	if( ProducerClass && ProducerClass->IsChildOf( UDataprepContentProducer::StaticClass() ) )
 	{
-		UDataprepContentProducer* Producer = NewObject< UDataprepContentProducer >( GetOutermost(), ProducerClass, NAME_None, RF_Transactional );
+		UDataprepContentProducer* Producer = NewObject< UDataprepContentProducer >( this, ProducerClass, NAME_None, RF_Transactional );
 		FAssetRegistryModule::AssetCreated( Producer );
 		Producer->MarkPackageDirty();
 
@@ -402,18 +402,25 @@ bool UDataprepAsset::RemoveProducer(int32 IndexToRemove)
 
 		Producers.RemoveAt( IndexToRemove );
 
-		if(Producers.Num() == 1)
+		// Array of producers superseded by removed producer
+		TArray<int32> ProducersToRevisit;
+		ProducersToRevisit.Reserve( Producers.Num() );
+
+		if( Producers.Num() == 1 )
 		{
 			Producers[0].SupersededBy = INDEX_NONE;
 		}
 		else if(Producers.Num() > 1)
 		{
 			// Update value stored in SupersededBy property where applicable
-			for( FDataprepAssetProducer& AssetProducer : Producers )
+			for( int32 Index = 0; Index < Producers.Num(); ++Index )
 			{
+				FDataprepAssetProducer& AssetProducer = Producers[Index];
+
 				if( AssetProducer.SupersededBy == IndexToRemove )
 				{
 					AssetProducer.SupersededBy = INDEX_NONE;
+					ProducersToRevisit.Add( Index );
 				}
 				else if( AssetProducer.SupersededBy > IndexToRemove )
 				{
@@ -426,6 +433,29 @@ bool UDataprepAsset::RemoveProducer(int32 IndexToRemove)
 
 		OnChanged.Broadcast( FDataprepAssetChangeType::ProducerRemoved, IndexToRemove );
 
+		// Update superseding status for producers depending on removed producer
+		bool bChangeAll = false;
+
+		for( int32 ProducerIndex : ProducersToRevisit )
+		{
+			bool bLocalChangeAll = false;
+			ValidateProducerChanges( ProducerIndex, bLocalChangeAll );
+			bChangeAll |= bLocalChangeAll;
+		}
+
+		// Notify observes on additional changes
+		if( bChangeAll )
+		{
+			OnChanged.Broadcast( FDataprepAssetChangeType::ProducerModified, INDEX_NONE );
+		}
+		else
+		{
+			for( int32 ProducerIndex : ProducersToRevisit )
+			{
+				OnChanged.Broadcast( FDataprepAssetChangeType::ProducerModified, ProducerIndex );
+			}
+		}
+
 		return true;
 	}
 
@@ -474,7 +504,7 @@ bool UDataprepAsset::ReplaceConsumer(UClass* NewConsumerClass)
 			DataprepAssetUtil::DeleteRegisteredAsset( Consumer );
 		}
 
-		Consumer = NewObject< UDataprepContentConsumer >( GetOutermost(), NewConsumerClass, NAME_None, RF_Transactional );
+		Consumer = NewObject< UDataprepContentConsumer >( this, NewConsumerClass, NAME_None, RF_Transactional );
 		check( Consumer );
 
 		FAssetRegistryModule::AssetCreated( Consumer );
diff --git a/Engine/Plugins/Experimental/DataPrepEditor/Source/DataPrepEditor/Private/DataPrepEditor.cpp b/Engine/Plugins/Experimental/DataPrepEditor/Source/DataPrepEditor/Private/DataPrepEditor.cpp
--- a/Engine/Plugins/Experimental/DataPrepEditor/Source/DataPrepEditor/Private/DataPrepEditor.cpp
+++ b/Engine/Plugins/Experimental/DataPrepEditor/Source/DataPrepEditor/Private/DataPrepEditor.cpp
@@ -13,8 +13,9 @@
 #include "DataprepActionAsset.h"
 #include "DataprepEditorLogCategory.h"
 #include "SchemaActions/DataprepOperationMenuActionCollector.h"
-#include "Widgets/SAssetsPreviewWidget.h"
+#include "Widgets/DataprepAssetView.h"
 #include "Widgets/SDataprepPalette.h"
+#include "Widgets/SAssetsPreviewWidget.h"
 
 #include "ActorEditorUtils.h"
 #include "AssetDeleteModel.h"
@@ -68,6 +69,8 @@ const FName FDataprepEditor::ScenePreviewTabId(TEXT("DataprepEditor_ScenePreview
 const FName FDataprepEditor::AssetPreviewTabId(TEXT("DataprepEditor_AssetPreview"));
 const FName FDataprepEditor::PaletteTabId(TEXT("DataprepEditor_Palette"));
 const FName FDataprepEditor::DetailsTabId(TEXT("DataprepEditor_Details"));
+const FName FDataprepEditor::DataprepAssetTabId(TEXT("DataprepEditor_Dataprep"));
+
 
 namespace DataprepEditorUtil
 {
@@ -253,6 +256,11 @@ void FDataprepEditor::RegisterTabSpawners(const TSharedRef<class FTabManager>& I
 		.SetGroup(WorkspaceMenuCategoryRef)
 		.SetIcon(FSlateIcon(FEditorStyle::GetStyleSetName(), "LevelEditor.Tabs.Details"));
 
+	InTabManager->RegisterTabSpawner(DataprepAssetTabId, FOnSpawnTab::CreateSP(this, &FDataprepEditor::SpawnTabDataprep))
+		.SetDisplayName(LOCTEXT("DataprepAssetTab", "Dataprep"))
+		.SetGroup(WorkspaceMenuCategoryRef)
+		.SetIcon(FSlateIcon(FEditorStyle::GetStyleSetName(), "LevelEditor.Tabs.Details"));
+
 	// Temp code for the nodes development
 	InTabManager->RegisterTabSpawner(PipelineGraphTabId, FOnSpawnTab::CreateSP(this, &FDataprepEditor::SpawnTabPipelineGraph))
 		.SetDisplayName(LOCTEXT("PipelineGraphTab", "Pipeline Graph"))
@@ -328,7 +336,7 @@ void FDataprepEditor::InitDataprepEditor(const EToolkitMode::Type Mode, const TS
 
 	CreateTabs();
 
-	const TSharedRef<FTabManager::FLayout> StandaloneDefaultLayout = FTabManager::NewLayout("Standalone_DataprepEditor_Layout_v0.2")
+	const TSharedRef<FTabManager::FLayout> StandaloneDefaultLayout = FTabManager::NewLayout("Standalone_DataprepEditor_Layout_v0.3")
 		->AddArea
 		(
 			FTabManager::NewPrimaryArea()->SetOrientation(Orient_Vertical)
@@ -386,9 +394,20 @@ void FDataprepEditor::InitDataprepEditor(const EToolkitMode::Type Mode, const TS
 				)
 				->Split
 				(
-					FTabManager::NewStack()
+					FTabManager::NewSplitter()->SetOrientation(Orient_Vertical)
 					->SetSizeCoefficient(0.25f)
-					->AddTab(DetailsTabId, ETabState::OpenedTab)
+					->Split
+					(
+						FTabManager::NewStack()
+						->SetSizeCoefficient(0.25f)
+						->AddTab(DataprepAssetTabId, ETabState::OpenedTab)
+					)
+					->Split
+					(
+						FTabManager::NewStack()
+						->SetSizeCoefficient(0.75f)
+						->AddTab(DetailsTabId, ETabState::OpenedTab)
+					)
 				)
 			)
 		);
@@ -400,9 +419,6 @@ void FDataprepEditor::InitDataprepEditor(const EToolkitMode::Type Mode, const TS
 	ExtendMenu();
 	ExtendToolBar();
 	RegenerateMenusAndToolbars();
-
-	// Set the details panel to inspect consumer and producers
-	OnShowSettings();
 }
 
 void FDataprepEditor::BindCommands()
@@ -431,12 +447,6 @@ void FDataprepEditor::BindCommands()
 		Commands.SaveScene,
 		FExecuteAction::CreateSP(this, &FDataprepEditor::OnSaveScene));
 
-	UICommandList->MapAction(
-		Commands.ShowDataprepSettings,
-		FExecuteAction::CreateSP(this, &FDataprepEditor::OnShowSettings),
-		FCanExecuteAction(),
-		FIsActionChecked::CreateSP(this, &FDataprepEditor::IsShowingSettings));
-
 	UICommandList->MapAction(
 		Commands.BuildWorld,
 		FExecuteAction::CreateSP( this, &FDataprepEditor::OnBuildWorld ),
@@ -746,12 +756,6 @@ void FDataprepEditor::ExtendToolBar()
 	{
 		static void FillToolbar(FToolBarBuilder& ToolbarBuilder, FDataprepEditor* ThisEditor)
 		{
-			ToolbarBuilder.BeginSection("Settings");
-			{
-				ToolbarBuilder.AddToolBarButton(FDataprepEditorCommands::Get().ShowDataprepSettings);
-			}
-			ToolbarBuilder.EndSection();
-
 			ToolbarBuilder.BeginSection("Run");
 			{
 				ToolbarBuilder.AddToolBarButton(FDataprepEditorCommands::Get().BuildWorld);
@@ -789,6 +793,8 @@ void FDataprepEditor::CreateTabs()
 		}
 	);
 
+	DataprepAssetView = SNew( SDataprepAssetView, DataprepAssetPtr.Get(), PipelineEditorCommands );
+
 	CreateScenePreviewTab();
 
 	// Create Details Panel
@@ -887,6 +893,22 @@ TSharedRef<SDockTab> FDataprepEditor::SpawnTabPalette(const FSpawnTabArgs & Args
 		];
 }
 
+TSharedRef<SDockTab> FDataprepEditor::SpawnTabDataprep(const FSpawnTabArgs & Args)
+{
+	check(Args.GetTabId() == DataprepAssetTabId);
+
+	return SNew(SDockTab)
+	.Label(LOCTEXT("DataprepEditor_DataprepTab_Title", "Dataprep"))
+	[
+		SNew(SBorder)
+		.Padding(2.f)
+		.BorderImage(FEditorStyle::GetBrush("ToolPanel.GroupBorder"))
+		[
+			DataprepAssetView.ToSharedRef()
+		]
+	];
+}
+
 void FDataprepEditor::UpdatePreviewPanels()
 {
 	// #ueent_todo: There should be a event triggered to inform listeners
diff --git a/Engine/Plugins/Experimental/DataPrepEditor/Source/DataPrepEditor/Private/DataPrepEditor.h b/Engine/Plugins/Experimental/DataPrepEditor/Source/DataPrepEditor/Private/DataPrepEditor.h
--- a/Engine/Plugins/Experimental/DataPrepEditor/Source/DataPrepEditor/Private/DataPrepEditor.h
+++ b/Engine/Plugins/Experimental/DataPrepEditor/Source/DataPrepEditor/Private/DataPrepEditor.h
@@ -89,8 +89,6 @@ public:
 private:
 	void BindCommands();
 	void OnSaveScene();
-	void OnShowSettings();
-	bool IsShowingSettings() const;
 	void OnBuildWorld();
 	void ResetBuildWorld();
 	void CleanPreviewWorld();
@@ -108,6 +106,7 @@ private:
 	TSharedRef<SDockTab> SpawnTabScenePreview(const FSpawnTabArgs& Args);
 	TSharedRef<SDockTab> SpawnTabPalette(const FSpawnTabArgs& Args);
 	TSharedRef<SDockTab> SpawnTabDetails(const FSpawnTabArgs& Args);
+	TSharedRef<SDockTab> SpawnTabDataprep(const FSpawnTabArgs& Args);
 
 	void TryInvokingDetailsTab(bool bFlash);
 
@@ -209,6 +208,7 @@ private:
 	static const FName AssetPreviewTabId;
 	static const FName PaletteTabId;
 	static const FName DetailsTabId;
+	static const FName DataprepAssetTabId;
 
 //Temp Code to allow us to work on the nodes while we don't have our own graph.
 public:
diff --git a/Engine/Plugins/Experimental/DataPrepEditor/Source/DataPrepEditor/Private/DataPrepEditorActions.cpp b/Engine/Plugins/Experimental/DataPrepEditor/Source/DataPrepEditor/Private/DataPrepEditorActions.cpp
--- a/Engine/Plugins/Experimental/DataPrepEditor/Source/DataPrepEditor/Private/DataPrepEditorActions.cpp
+++ b/Engine/Plugins/Experimental/DataPrepEditor/Source/DataPrepEditor/Private/DataPrepEditorActions.cpp
@@ -11,7 +11,6 @@ void FDataprepEditorCommands::RegisterCommands()
 	// end of temp code for nodes development
 
 	UI_COMMAND(SaveScene, "Save Datasmith Scene", "To be filled.", EUserInterfaceActionType::ToggleButton, FInputChord());
-	UI_COMMAND(ShowDataprepSettings, "Settings", "Edit the data prep settings.", EUserInterfaceActionType::ToggleButton, FInputChord());
 	UI_COMMAND(ShowDatasmithSceneSettings, "Scene Settings", "Edit the scene settings.", EUserInterfaceActionType::ToggleButton, FInputChord());
 	UI_COMMAND(BuildWorld, "Import", "Load external data sources referenced by producers into the Dataprep Editor", EUserInterfaceActionType::ToggleButton, FInputChord());
 	UI_COMMAND(ExecutePipeline, "Execute", "Execute the Dataprep graph on the imported data", EUserInterfaceActionType::ToggleButton, FInputChord());
diff --git a/Engine/Plugins/Experimental/DataPrepEditor/Source/DataPrepEditor/Private/DataPrepEditorActions.h b/Engine/Plugins/Experimental/DataPrepEditor/Source/DataPrepEditor/Private/DataPrepEditorActions.h
--- a/Engine/Plugins/Experimental/DataPrepEditor/Source/DataPrepEditor/Private/DataPrepEditorActions.h
+++ b/Engine/Plugins/Experimental/DataPrepEditor/Source/DataPrepEditor/Private/DataPrepEditorActions.h
@@ -30,7 +30,6 @@ public:
 
 	/**  */
 	TSharedPtr< FUICommandInfo > SaveScene;
-	TSharedPtr< FUICommandInfo > ShowDataprepSettings;
 	TSharedPtr< FUICommandInfo > ShowDatasmithSceneSettings;
 	TSharedPtr< FUICommandInfo > BuildWorld;
 	TSharedPtr< FUICommandInfo > ExecutePipeline;
diff --git a/Engine/Plugins/Experimental/DataPrepEditor/Source/DataPrepEditor/Private/DataprepDetailsTab.cpp b/Engine/Plugins/Experimental/DataPrepEditor/Source/DataPrepEditor/Private/DataprepDetailsTab.cpp
--- a/Engine/Plugins/Experimental/DataPrepEditor/Source/DataPrepEditor/Private/DataprepDetailsTab.cpp
+++ b/Engine/Plugins/Experimental/DataPrepEditor/Source/DataPrepEditor/Private/DataprepDetailsTab.cpp
@@ -68,32 +68,6 @@ TSharedRef<SDockTab> FDataprepEditor::SpawnTabDetails(const FSpawnTabArgs & Args
 void FDataprepEditor::CreateDetailsViews()
 {
 	DetailsView = SNew( SGraphNodeDetailsWidget );
-	DataprepAssetView = SNew( SDataprepAssetView, DataprepAssetPtr.Get(), PipelineEditorCommands );
-}
-
-void FDataprepEditor::OnShowSettings()
-{
-	if ( !DetailsTabPtr.IsValid() )
-	{
-		TryInvokingDetailsTab( false );
-	}
-
-	if ( DetailsTabPtr.IsValid() && DetailsTabPtr.Pin()->GetContent() != DataprepAssetView )
-	{
-		DetailsTabPtr.Pin()->SetContent( DataprepAssetView.ToSharedRef() );
-	}
-
-	TryInvokingDetailsTab( true );
-}
-
-bool FDataprepEditor::IsShowingSettings() const
-{
-	if ( DetailsTabPtr.IsValid() )
-	{
-		return DetailsTabPtr.Pin()->GetContent() == DataprepAssetView;
-	}
-
-	return false;
 }
 
 void FDataprepEditor::OnPipelineEditorSelectionChanged(const TSet<UObject*>& SelectedNodes)
diff --git a/Engine/Plugins/Experimental/DataPrepEditor/Source/DataPrepEditor/Private/Widgets/DataprepAssetView.cpp b/Engine/Plugins/Experimental/DataPrepEditor/Source/DataPrepEditor/Private/Widgets/DataprepAssetView.cpp
--- a/Engine/Plugins/Experimental/DataPrepEditor/Source/DataPrepEditor/Private/Widgets/DataprepAssetView.cpp
+++ b/Engine/Plugins/Experimental/DataPrepEditor/Source/DataPrepEditor/Private/Widgets/DataprepAssetView.cpp
@@ -16,6 +16,7 @@
 #include "K2Node_AddComponent.h"
 #include "PropertyEditorModule.h"
 
+#include "DetailLayoutBuilder.h"
 #include "EditorFontGlyphs.h"
 #include "Modules/ModuleManager.h"
 #include "Widgets/Colors/SColorBlock.h"
@@ -26,6 +27,7 @@
 #include "Widgets/Layout/SBorder.h"
 #include "Widgets/Layout/SGridPanel.h"
 #include "Widgets/Layout/SScrollBox.h"
+#include "Widgets/Layout/SScrollBar.h"
 #include "Widgets/Layout/SSeparator.h"
 #include "Widgets/Layout/SSpacer.h"
 #include "Widgets/Text/SInlineEditableTextBlock.h"
@@ -34,7 +36,14 @@
 #define LOCTEXT_NAMESPACE "DataprepAssetView"
 
 const float IndentSize = 12;
-const char* DataprepTabsFontName = "FontAwesome.11";
+
+namespace DataprepEditorUtils
+{
+	FSlateFontInfo GetGlyphFont()
+	{
+		return FEditorStyle::Get().GetFontStyle( "FontAwesome.11" );
+	}
+}
 
 // Inspired from SKismetInspector::Construct
 void SGraphNodeDetailsWidget::Construct(const FArguments& InArgs)
@@ -369,14 +378,6 @@ TSharedRef<SWidget> SProducerStackEntryTableRow::GetInputMainWidget()
 		return SNullWidget::NullWidget;
 	}
 
-	auto ToggleEntry = [this, ProducerStackEntry]()
-	{
-		ProducerStackEntry->ToggleProducer();
-		CheckBox->SetText( ProducerStackEntry->WillBeRun() ? FEditorFontGlyphs::Check_Square : FEditorFontGlyphs::Square );
-
-		return FReply::Handled();
-	};
-
 	auto DeleteEntry = [ProducerStackEntry]()
 	{
 		ProducerStackEntry->RemoveProducer();
@@ -386,6 +387,8 @@ TSharedRef<SWidget> SProducerStackEntryTableRow::GetInputMainWidget()
 	// Padding for check and delete buttons to center them on the first line of the detail view
 	const FMargin ButtonPadding( 0.0f, 10.0f, 0.0f, 0.0f );
 
+	TSharedPtr<STextBlock> StatusText;
+
 	TSharedPtr<SWidget> Widget = SNew(SBorder)
 	.BorderImage(FEditorStyle::GetBrush("NoBrush"))
 	.Padding(5.0f)
@@ -397,19 +400,10 @@ TSharedRef<SWidget> SProducerStackEntryTableRow::GetInputMainWidget()
 		.Padding( ButtonPadding )
 		.AutoWidth()
 		[
-			SNew(SButton)
-			.ButtonStyle(FEditorStyle::Get(), "HoverHintOnly")
-			.ToolTipText(LOCTEXT("CheckToolTip", "Include or exclude this producer to the creation of the world "))
-			.IsFocusable(false)
-			.OnClicked_Lambda( ToggleEntry )
-			.VAlign( VAlign_Top )
-			.Content()
-			[
-				SAssignNew(CheckBox, STextBlock)
-				.Font(FEditorStyle::Get().GetFontStyle(DataprepTabsFontName))
-				.ColorAndOpacity( FLinearColor::White )
-				.Text( ProducerStackEntry->WillBeRun() ? FEditorFontGlyphs::Check_Square : FEditorFontGlyphs::Square )
-			]
+			SAssignNew(StatusText, STextBlock)
+			.Font( DataprepEditorUtils::GetGlyphFont() )
+			.ColorAndOpacity( this, &SProducerStackEntryTableRow::GetStatusColorAndOpacity )
+			.Text( FEditorFontGlyphs::Exclamation_Triangle )
 		]
 		// Input entry label
 		+ SHorizontalBox::Slot()
@@ -427,26 +421,42 @@ TSharedRef<SWidget> SProducerStackEntryTableRow::GetInputMainWidget()
 		[
 			SNew(SButton)
 			.ButtonStyle(FEditorStyle::Get(), "HoverHintOnly")
-			.ToolTipText(LOCTEXT("DeleteToolTip", "Delete this producer"))
+			.ToolTipText(LOCTEXT("ProducerStackEntryTableRow_DeleteToolTip", "Delete this producer"))
 			.IsFocusable(false)
 			.OnClicked_Lambda( DeleteEntry )
 			.VAlign( VAlign_Top )
 			.Content()
 			[
 				SNew(STextBlock)
-				.Font(FEditorStyle::Get().GetFontStyle(DataprepTabsFontName))
+				.Font( DataprepEditorUtils::GetGlyphFont() )
 				.ColorAndOpacity( FLinearColor::White )
 				.Text( FEditorFontGlyphs::Trash )
 			]
 		]
 	];
 
-	// Disable check box if associated producer is superseeded
-	CheckBox->SetEnabled( ProducerStackEntry->WillBeRun() );
+	StatusText->SetToolTipText( TAttribute<FText>( this, &SProducerStackEntryTableRow::GetStatusTooltipText ) );
 
 	return Widget.ToSharedRef();
 }
 
+FSlateColor SProducerStackEntryTableRow::GetStatusColorAndOpacity() const
+{
+	FProducerStackEntryPtr ProducerStackEntry = Node.Pin();
+	return  ( ProducerStackEntry.IsValid() && ProducerStackEntry->WillBeRun() ) ? FLinearColor::Transparent : FLinearColor::Red;
+}
+
+FText SProducerStackEntryTableRow::GetStatusTooltipText() const
+{
+	FProducerStackEntryPtr ProducerStackEntry = Node.Pin();
+	if( !ProducerStackEntry.IsValid() )
+	{
+		return LOCTEXT( "ProducerStackEntryTableRow_StatusTextTooltip_Invalid", "The producer is not valid");
+	}
+
+	return  ProducerStackEntry->WillBeRun() ? FText() : LOCTEXT( "ProducerStackEntryTableRow_StatusTextTooltip_Superseded", "This producer is superseded by another one and will be skipped when run.");
+}
+
 void SDataprepAssetView::Construct( const FArguments& InArgs, UDataprepAsset* InDataprepAssetPtr, TSharedPtr<FUICommandList>& CommandList )
 {
 	check( InDataprepAssetPtr );
@@ -489,26 +499,6 @@ void SDataprepAssetView::Construct( const FArguments& InArgs, UDataprepAsset* In
 		SelectedConsumerDescription = MakeShared<FString>( FString() );
 	}
 
-	auto CheckEntry = [ this ]()
-	{
-		for(int32 Index = 0; Index < DataprepAssetPtr->GetProducersCount(); ++Index)
-		{
-			DataprepAssetPtr->EnableProducer( Index, !bIsChecked );
-		}
-
-
-		if( DataprepAssetPtr->EnableAllProducers( !bIsChecked ) )
-		{
-			bIsChecked = !bIsChecked;
-		}
-
-		CheckBox->SetText( bIsChecked ? FEditorFontGlyphs::Check_Square : FEditorFontGlyphs::Square );
-
-		TreeView->Refresh();
-
-		return FReply::Handled();
-	};
-
 	TSharedPtr<SWidget> AddNewMenu = SNew(SComboButton)
 	.ComboButtonStyle(FEditorStyle::Get(), "ToolbarComboButton")
 	.ForegroundColor(FLinearColor::White)
@@ -524,13 +514,14 @@ void SDataprepAssetView::Construct( const FArguments& InArgs, UDataprepAsset* In
 		.HAlign(HAlign_Center)
 		[
 			SAssignNew(CheckBox, STextBlock)
-			.Font(FEditorStyle::Get().GetFontStyle(DataprepTabsFontName))
+			.Font( DataprepEditorUtils::GetGlyphFont() )
 			.ColorAndOpacity( FLinearColor::White )
-			.Text(FText::FromString(FString( TEXT( "\xf055" ) ) ) )
+			.Text( FEditorFontGlyphs::Plus_Circle )
 		]
 	];
 
 	TreeView = SNew(SProducerStackEntryTreeView, this, DataprepAssetPtr.Get() );
+	TSharedRef<SScrollBar> ScrollBar = SNew(SScrollBar);
 
 	// #ueent_todo: Look at changing the border brushes to add color to this stuff
 	ChildSlot
@@ -539,138 +530,125 @@ void SDataprepAssetView::Construct( const FArguments& InArgs, UDataprepAsset* In
 		.Padding(4.0f)
 		.BorderImage(FEditorStyle::GetBrush("ToolPanel.GroupBorder"))
 		[
-			SNew(SVerticalBox)
-			+ SVerticalBox::Slot()
-			.AutoHeight()
-			[
-				SNew(SSpacer)
-				.Size( FVector2D( 200, 100 ) )		
-			]
-			+ SVerticalBox::Slot()
-			.AutoHeight()
+			SNew(SHorizontalBox)
+			+SHorizontalBox::Slot()
+			.FillWidth(1)
 			[
-				SNew(SHorizontalBox)
-				+ SHorizontalBox::Slot()
-				.FillWidth(1.0f)
-				.Padding( 0, 10, 0, 0 )
-				.HAlign( EHorizontalAlignment::HAlign_Center )
-				[
-					// #ueent_todo: make color block's width vary with parent widget
-					SNew(SColorBlock)
-					.Color( FLinearColor( 0.9f, 0.9f, 0.9f ) )
-					.Size( FVector2D( 3000, 1 ) )
-				]
-			]
-			+ SVerticalBox::Slot()
-			.AutoHeight()
-			[
-				SNew(SHorizontalBox)
-				+ SHorizontalBox::Slot()
-				.FillWidth(1.0f)
-				.HAlign(EHorizontalAlignment::HAlign_Left)
-				.VAlign(VAlign_Center)
-				.Padding( 5.0f )
-				[
-					SNew(STextBlock)
-					.Text(LOCTEXT("DataprepAssetView_Consumer_label", "Output"))
-					.MinDesiredWidth( 200 )
-				]
-			]
-			+ SVerticalBox::Slot()
-			.AutoHeight()
-			[
-				SNew(SHorizontalBox)
-				+ SHorizontalBox::Slot()
-				.FillWidth(1.0f)
-				.HAlign(EHorizontalAlignment::HAlign_Left)
-				.VAlign(VAlign_Center)
-				[
-					SAssignNew( ProducerSelector, STextComboBox )
-					.OptionsSource( &ConsumerDescriptionList )
-					.OnSelectionChanged( this, &SDataprepAssetView::OnNewConsumerSelected )
-					.InitiallySelectedItem( SelectedConsumerDescription )
-				]
-			]
-			+ SVerticalBox::Slot()
-			.AutoHeight()
-			[
-				SAssignNew( ConsumerWidget, SDataprepConsumerWidget )
-				.DataprepConsumer( DataprepAssetPtr->GetConsumer() )
-			]
-			// Section for producers
-			+ SVerticalBox::Slot()
-			.AutoHeight()
-			[
-				SNew(SHorizontalBox)
-				+ SHorizontalBox::Slot()
-				[
-					SNew(SSpacer)
-					.Size( FVector2D( 200, 20 ) )		
-				]
-			]
-			+ SVerticalBox::Slot()
-			.AutoHeight()
-			[
-				SNew(SHorizontalBox)
-				+ SHorizontalBox::Slot()
-				.FillWidth(1.0f)
-				.Padding( 0, 10, 0, 0 )
-				.HAlign( EHorizontalAlignment::HAlign_Center )
-				[
-					SNew(SColorBlock)
-					.Color( FLinearColor( 0.9f, 0.9f, 0.9f ) )
-					.Size( FVector2D( 2000, 1 ) )
-				]
-			]
-			+ SVerticalBox::Slot()
-			.Padding( 5.0f )
-			.AutoHeight()
-			[
-				SNew(SHorizontalBox)
-				// Check button
-				+ SHorizontalBox::Slot()
-				.HAlign(HAlign_Center)
-				.AutoWidth()
+
+				SNew(SScrollBox)
+				.ExternalScrollbar(ScrollBar)
+				+ SScrollBox::Slot()
 				[
-					SNew(SButton)
-					.ButtonStyle(FEditorStyle::Get(), "HoverHintOnly")
-					.ToolTipText(LOCTEXT("CheckToolTip", "Include or exclude this producer to the creation of the world "))
-					.IsFocusable(false)
-					.OnClicked_Lambda( CheckEntry )
-					.VAlign(VAlign_Center)
-					.Content()
+
+					SNew(SVerticalBox)
+					// Section for producers
+					+ SVerticalBox::Slot()
+					.Padding( 5.0f )
+					.AutoHeight()
 					[
-						SAssignNew(CheckBox, STextBlock)
-						.Font(FEditorStyle::Get().GetFontStyle(DataprepTabsFontName))
-						.ColorAndOpacity( FLinearColor::White )
-						.Text( bIsChecked ? FEditorFontGlyphs::Check_Square : FEditorFontGlyphs::Square )
+						SNew(SHorizontalBox)
+						+ SHorizontalBox::Slot()
+						.FillWidth(1.0f)
+						.HAlign(EHorizontalAlignment::HAlign_Left)
+						.VAlign(VAlign_Center)
+						.Padding( FMargin( 10, 0, 0, 0 ) )
+						[
+							SNew(STextBlock)
+							.Font( IDetailLayoutBuilder::GetDetailFontBold() )
+							.Text(LOCTEXT("DataprepAssetView_Producers_label", "Inputs"))
+						]
+						+ SHorizontalBox::Slot()
+						.AutoWidth()
+						.HAlign(EHorizontalAlignment::HAlign_Right)
+						.Padding(0, 0, 2, 0)
+						[
+							AddNewMenu.ToSharedRef()
+						]
 					]
-				]
-				+ SHorizontalBox::Slot()
-				.FillWidth(1.0f)
-				.HAlign(EHorizontalAlignment::HAlign_Left)
-				.VAlign(VAlign_Center)
-				[
-					SNew(STextBlock)
-					.Text(LOCTEXT("DataprepAssetView_Producers_label", "Inputs"))
-				]
-				+ SHorizontalBox::Slot()
-				.AutoWidth()
-				.HAlign(EHorizontalAlignment::HAlign_Right)
-				.Padding(0, 0, 2, 0)
-				[
-					AddNewMenu.ToSharedRef()
+					+ SVerticalBox::Slot()
+					.Padding( 5.0f )
+					.AutoHeight()
+					[
+						TreeView.ToSharedRef()
+					]
+					+ SVerticalBox::Slot()
+					.AutoHeight()
+					[
+						SNew(SHorizontalBox)
+						+ SHorizontalBox::Slot()
+						.FillWidth(1.0f)
+						.Padding( 0, 10, 0, 0 )
+						.HAlign( EHorizontalAlignment::HAlign_Center )
+						[
+							// #ueent_todo: make color block's width vary with parent widget
+							SNew(SColorBlock)
+							.Color( FLinearColor( 0.9f, 0.9f, 0.9f ) )
+							.Size( FVector2D( 3000, 1 ) )
+						]
+					]
+					+ SVerticalBox::Slot()
+					.AutoHeight()
+					[
+						SNew(SHorizontalBox)
+						+ SHorizontalBox::Slot()
+						[
+							SNew(SSpacer)
+							.Size( FVector2D( 200, 10 ) )		
+						]
+					]
+					// Section for consumer
+					+ SVerticalBox::Slot()
+					.AutoHeight()
+					[
+						SNew(SHorizontalBox)
+						+ SHorizontalBox::Slot()
+						.FillWidth(1.0f)
+						.HAlign(EHorizontalAlignment::HAlign_Left)
+						.VAlign(VAlign_Center)
+						.Padding( FMargin( 10, 0, 0, 0 ) )
+						[
+							SNew(STextBlock)
+							.Text(LOCTEXT("DataprepAssetView_Consumer_label", "Output"))
+							.MinDesiredWidth( 200 )
+							.Font( IDetailLayoutBuilder::GetDetailFontBold() )
+						]
+						+ SHorizontalBox::Slot()
+						.AutoWidth()
+						.HAlign(EHorizontalAlignment::HAlign_Right)
+						.Padding(0, 0, 2, 0)
+						[
+							SAssignNew( ProducerSelector, STextComboBox )
+							.OptionsSource( &ConsumerDescriptionList )
+							.OnSelectionChanged( this, &SDataprepAssetView::OnNewConsumerSelected )
+							.InitiallySelectedItem( SelectedConsumerDescription )
+						]
+					]
+					+ SVerticalBox::Slot()
+					.AutoHeight()
+					[
+						SAssignNew( ConsumerWidget, SDataprepConsumerWidget )
+						.DataprepConsumer( DataprepAssetPtr->GetConsumer() )
+					]
+					+ SVerticalBox::Slot()
+					.AutoHeight()
+					[
+						SNew(SHorizontalBox)
+						+ SHorizontalBox::Slot()
+						[
+							SNew(SSpacer)
+							.Size( FVector2D( 200, 10 ) )		
+						]
+					]
+					// Section for consumer
 				]
 			]
-			+ SVerticalBox::Slot()
-			.Padding( 5.0f )
-			.AutoHeight()
+			+SHorizontalBox::Slot()
+			.AutoWidth()
 			[
-				SNew(SScrollBox)
-				+ SScrollBox::Slot()
-				.Padding(0, 3, 0, 3)
+				SNew( SBox )
+				.WidthOverride( FOptionalSize( 16 ) )
 				[
-					TreeView.ToSharedRef()
+					ScrollBar
 				]
 			]
 		]
diff --git a/Engine/Plugins/Experimental/DataPrepEditor/Source/DataPrepEditor/Private/Widgets/DataprepAssetView.h b/Engine/Plugins/Experimental/DataPrepEditor/Source/DataPrepEditor/Private/Widgets/DataprepAssetView.h
--- a/Engine/Plugins/Experimental/DataPrepEditor/Source/DataPrepEditor/Private/Widgets/DataprepAssetView.h
+++ b/Engine/Plugins/Experimental/DataPrepEditor/Source/DataPrepEditor/Private/Widgets/DataprepAssetView.h
@@ -95,9 +95,12 @@ public:
 		return Node.Pin();
 	}
 
+private:
+	FSlateColor GetStatusColorAndOpacity() const;
+	FText GetStatusTooltipText() const;
+
 private:
 	mutable TWeakPtr<FProducerStackEntry> Node;
-	TSharedPtr< STextBlock > CheckBox;
 };
 
 class SProducerStackEntryTreeView : public STreeView<FProducerStackEntryRef>
diff --git a/Engine/Plugins/Experimental/DataPrepEditor/Source/DataPrepEditor/Private/Widgets/DataprepWidgets.cpp b/Engine/Plugins/Experimental/DataPrepEditor/Source/DataPrepEditor/Private/Widgets/DataprepWidgets.cpp
--- a/Engine/Plugins/Experimental/DataPrepEditor/Source/DataPrepEditor/Private/Widgets/DataprepWidgets.cpp
+++ b/Engine/Plugins/Experimental/DataPrepEditor/Source/DataPrepEditor/Private/Widgets/DataprepWidgets.cpp
@@ -5,6 +5,7 @@
 #include "DataPrepContentConsumer.h"
 
 #include "ContentBrowserModule.h"
+#include "DetailLayoutBuilder.h"
 #include "Dialogs/DlgPickPath.h"
 #include "EditorStyleSet.h"
 #include "Framework/MultiBox/MultiBoxBuilder.h"
@@ -96,12 +97,12 @@ TSharedRef<SWidget> SDataprepConsumerWidget::BuildWidget()
 
 	TSharedRef<SWidget> BrowseButton = PropertyCustomizationHelpers::MakeBrowseButton( FSimpleDelegate::CreateSP( this, &SDataprepConsumerWidget::OnBrowseContentFolder ) );
 
-	return SNew(SBorder)
-	.BorderImage( FCoreStyle::Get().GetBrush("ToolPanel.GroupBorder") )
+	TSharedRef<SWidget> Widget = SNew(SBorder)
+	.BorderImage( FCoreStyle::Get().GetBrush("NoBrush") )
 	[
 		SNew(SVerticalBox)
 		+ SVerticalBox::Slot()
-		.Padding(5.0f)
+		.Padding( FMargin( 20, 0, 0, 0 ) )
 		[
 			SNew(SGridPanel)
 			.FillColumn(0, 0.25f)
@@ -115,32 +116,29 @@ TSharedRef<SWidget> SDataprepConsumerWidget::BuildWidget()
 				.FillWidth(1.0f)
 				[
 					SNew(STextBlock)
-					.Text(LOCTEXT("DataprepSlateHelper_ContentFolderLabel", "Content Folder"))
-					.Margin(FMargin( 5.0f, 5.0f, 0.0f, 0.0f ) )
+					.Text(LOCTEXT("DataprepSlateHelper_ContentFolderLabel", "Folder"))
+					.Font( IDetailLayoutBuilder::GetDetailFont() )
 				]
 			]
 			+ SGridPanel::Slot(1, 0)
 			.Padding(0.0f, 5.0f)
 			[
-				SNew(SBox)
+				SNew(SHorizontalBox)
+				+ SHorizontalBox::Slot()
+				.VAlign(VAlign_Center)
+				.FillWidth(1.0f)
 				[
-					SNew(SHorizontalBox)
-					+ SHorizontalBox::Slot()
-					.VAlign(VAlign_Center)
-					.FillWidth(1.0f)
-					[
-						SAssignNew(ContentFolderTextBox, SEditableTextBox)
-						.Text( FText::FromString( DataprepConsumer->GetTargetContentFolder() ) )
-						.HintText( LOCTEXT("DataprepSlateHelper_ContentFolderHintText", "Set the content folder to save in") )
-						.IsReadOnly( true )
-					]
-					+ SHorizontalBox::Slot()
-					.VAlign(VAlign_Center)
-					.HAlign(HAlign_Center)
-					[
-						BrowseButton
-					]
-				]			
+					SAssignNew(ContentFolderTextBox, SEditableTextBox)
+					.Font( IDetailLayoutBuilder::GetDetailFont() )
+					.HintText( LOCTEXT("DataprepSlateHelper_ContentFolderHintText", "Set the content folder to save in") )
+					.IsReadOnly( true )
+				]
+				+ SHorizontalBox::Slot()
+				.VAlign(VAlign_Center)
+				.HAlign(HAlign_Center)
+				[
+					BrowseButton
+				]
 			]
 			+ SGridPanel::Slot(0, 1)
 			[
@@ -152,7 +150,7 @@ TSharedRef<SWidget> SDataprepConsumerWidget::BuildWidget()
 				[
 					SNew(STextBlock)
 					.Text(LOCTEXT("DataprepSlateHelper_LevelNameLabel", "Level"))
-					.Margin(FMargin( 5.0f, 5.0f, 0.0f, 0.0f ) )
+					.Font( IDetailLayoutBuilder::GetDetailFont() )
 				]
 			]
 			+ SGridPanel::Slot(1, 1)
@@ -166,12 +164,17 @@ TSharedRef<SWidget> SDataprepConsumerWidget::BuildWidget()
 				[
 					SAssignNew(LevelTextBox, SEditableTextBox)
 					.Text( FText::FromString( DataprepConsumer->GetLevelName() ) )
+					.Font( IDetailLayoutBuilder::GetDetailFont() )
 					.HintText( LOCTEXT("DataprepSlateHelper_LevelNameHintText", "Current will be used") )
 					.OnTextCommitted(  this, &SDataprepConsumerWidget::OnLevelNameChanged  )
 				]
 			]
 		]
 	];
+
+	UpdateContentFolderText();
+
+	return Widget;
 }
 
 TSharedRef<SWidget> SDataprepConsumerWidget::BuildNullWidget()
@@ -190,7 +193,7 @@ TSharedRef<SWidget> SDataprepConsumerWidget::BuildNullWidget()
 			.FillWidth(1.0f)
 			[
 				SNew(STextBlock)
-				.Font(FCoreStyle::GetDefaultFontStyle("Bold", 11))
+				.Font( IDetailLayoutBuilder::GetDetailFont() )
 				.Text(LOCTEXT("DataprepSlateHelper_ContentFolderLabel", "Error: Not a valid consumer"))
 				.Margin(FMargin( 5.0f, 5.0f, 0.0f, 0.0f ) )
 				.ColorAndOpacity( FLinearColor(1,0,0,1) )
@@ -204,7 +207,12 @@ void SDataprepConsumerWidget::UpdateContentFolderText()
 	if(UDataprepContentConsumer* Consumer = DataprepConsumer.Get())
 	{
 		FString TargetContentFolder( Consumer->GetTargetContentFolder() );
-		if( TargetContentFolder.StartsWith( TEXT( "/Game" ) ) )
+
+		if( TargetContentFolder.IsEmpty() )
+		{
+			TargetContentFolder = TEXT("/Content");
+		}
+		else if( TargetContentFolder.StartsWith( TEXT( "/Game" ) ) )
 		{
 			TargetContentFolder = TargetContentFolder.Replace( TEXT( "/Game" ), TEXT( "/Content" ) );
 		}
@@ -513,7 +521,7 @@ void SDataprepDetailsView::Construct()
 				.FillWidth(1.0f)
 				[
 					SNew(STextBlock)
-					.Font(FCoreStyle::GetDefaultFontStyle("Bold", 11))
+					.Font( IDetailLayoutBuilder::GetDetailFontBold() )
 					.Text( ErrorText )
 					.Margin(FMargin( 5.0f, 5.0f, 0.0f, 0.0f ) )
 					.ColorAndOpacity( FLinearColor(1,0,0,1) )
