Fixed the dataprep build issues for a monolithic editor and removed some dead code.

diff --git a/Engine/Plugins/Experimental/DataPrepEditor/Source/DataPrepCore/Private/DataPrepAsset.cpp b/Engine/Plugins/Experimental/DataPrepEditor/Source/DataPrepCore/Private/DataPrepAsset.cpp
--- a/Engine/Plugins/Experimental/DataPrepEditor/Source/DataPrepCore/Private/DataPrepAsset.cpp
+++ b/Engine/Plugins/Experimental/DataPrepEditor/Source/DataPrepCore/Private/DataPrepAsset.cpp
@@ -30,67 +30,6 @@ namespace DataprepAssetUtil
 	}
 }
 
-// FDataprepAssetAction =================================================================
-
-FDataprepAssetAction::~FDataprepAssetAction()
-{
-	UnbindDataprepAssetFromAction();
-}
-
-FDataprepAssetAction& FDataprepAssetAction::operator=(const FDataprepAssetAction& Other)
-{
-	bIsEnabled = Other.bIsEnabled;
-	DataprepAssetPtr = Other.DataprepAssetPtr;
-	SetActionAsset( Other.ActionAsset );
-	return *this;
-}
-
-FDataprepAssetAction& FDataprepAssetAction::operator=(FDataprepAssetAction&& Other)
-{
-	bIsEnabled = Other.bIsEnabled;
-	DataprepAssetPtr = Other.DataprepAssetPtr;
-	SetActionAsset( Other.ActionAsset );
-	Other.SetActionAsset( nullptr );
-	return *this;
-}
-
-void FDataprepAssetAction::SetActionAsset(UDataprepActionAsset* InActionAsset)
-{
-	if ( ActionAsset != InActionAsset )
-	{
-		UnbindDataprepAssetFromAction();
-	}
-
-	ActionAsset = InActionAsset;
-	BindDataprepAssetToAction();
-
-}
-
-void FDataprepAssetAction::BindDataprepAssetToAction()
-{
-	if ( ActionAsset )
-	{
-		OnOperationOrderChangedHandle = ActionAsset->GetOnStepsOrderChanged().AddRaw( this, &FDataprepAssetAction::OnActionOperationsOrderChanged );
-	}
-}
-
-void FDataprepAssetAction::UnbindDataprepAssetFromAction()
-{
-	if ( ActionAsset && OnOperationOrderChangedHandle.IsValid() )
-	{
-		ActionAsset->GetOnStepsOrderChanged().Remove( OnOperationOrderChangedHandle );
-	}
-}
-
-void FDataprepAssetAction::OnActionOperationsOrderChanged()
-{
-	UDataprepAsset* DataprepAsset = DataprepAssetPtr.Get();
-	if (DataprepAsset)
-	{
-		DataprepAsset->GetOnActionOperationsOrderChanged().Broadcast( ActionAsset );
-	}
-}
-
 
 // UDataprepAsset =================================================================
 
@@ -103,27 +42,6 @@ UDataprepAsset::UDataprepAsset()
 	// Temp code for the nodes development
 	DataprepRecipeBP = nullptr;
 	// end of temp code for nodes development
-
-#ifdef WITH_EDITOR
-	OnAssetDeletedHandle = FEditorDelegates::OnAssetsDeleted.AddLambda( [this](const TArray<UClass *>& DeletedClasses)
-		{
-			for (UClass* Class : DeletedClasses)
-			{
-				if (Class->IsChildOf<UDataprepAsset>())
-				{
-					RemoveInvalidActions();
-					break;
-				}
-			}
-		});
-#endif //WITH_EDITOR
-}
-
-UDataprepAsset::~UDataprepAsset()
-{
-#ifdef WITH_EDITOR
-	FEditorDelegates::OnAssetsDeleted.Remove( OnAssetDeletedHandle );
-#endif //WITH_EDITOR
 }
 
 void UDataprepAsset::PostInitProperties()
@@ -195,115 +113,6 @@ void UDataprepAsset::Serialize( FArchive& Ar )
 	}
 }
 
-int32 UDataprepAsset::AddAction()
-{
-	UDataprepActionAsset* Action = NewObject< UDataprepActionAsset >( this );
-	Actions.Emplace( Action, true, *this );
-	OnActionsOrderChanged.Broadcast();
-	return Actions.Num();
-}
-
-UDataprepActionAsset* UDataprepAsset::GetAction(int32 Index)
-{
-	// const_cast to avoid code duplication
-	return const_cast< UDataprepActionAsset* >( static_cast< const UDataprepAsset* >( this )->GetAction( Index ) );
-}
-
-const UDataprepActionAsset* UDataprepAsset::GetAction(int32 Index) const
-{
-	if ( Actions.IsValidIndex( Index ) )
-	{
-		return Actions[ Index ].GetActionAsset();
-	}
-
-	UE_LOG( LogDataprepCore, Error, TEXT("DataprepAsset::GetAction: the Index is out of range") );
-	return nullptr;
-}
-
-int32 UDataprepAsset::GetActionsCount() const
-{
-	return Actions.Num();
-}
-
-bool UDataprepAsset::IsActionEnabled(int32 Index) const
-{
-	if ( Actions.IsValidIndex( Index ) )
-	{
-		return Actions[ Index ].IsEnabled();
-	}
-
-	UE_LOG( LogDataprepCore, Error, TEXT("DataprepAsset::IsActionEnabled: the Index is out of range") );
-	return false;
-}
-
-void UDataprepAsset::EnableAction(int32 Index, bool bEnable)
-{
-	if ( Actions.IsValidIndex( Index ) )
-	{
-		Actions[ Index ].Enable( bEnable );
-	}
-	else
-	{
-		UE_LOG( LogDataprepCore, Error, TEXT("DataprepAsset::EnableAction: the Index is out of range") );
-	}
-}
-
-bool UDataprepAsset::MoveAction(int32 ActionIndex, int32 DestinationIndex)
-{
-	if ( DataprepCoreUtils::MoveArrayElement( Actions, ActionIndex, DestinationIndex ) )
-	{
-		OnActionsOrderChanged.Broadcast();
-		return true;
-	}
-
-	if ( !Actions.IsValidIndex( ActionIndex ) )
-	{
-		UE_LOG( LogDataprepCore, Error, TEXT("DataprepAsset::MoveAction: the ActionIndex is out of range") );
-	}
-	if ( !Actions.IsValidIndex( DestinationIndex ) )
-	{
-		UE_LOG( LogDataprepCore, Error, TEXT("DataprepAsset::MoveAction: the Destination Index is out of range") );
-	}
-	if ( ActionIndex == DestinationIndex )
-	{
-		UE_LOG( LogDataprepCore, Error, TEXT("DataprepAsset::MoveAction: an action shouldn't be move at the location it currently is") );
-	}
-	return false;
-}
-
-bool UDataprepAsset::RemoveAction(int32 Index)
-{
-	if ( Actions.IsValidIndex( Index ) )
-	{
-		Actions.RemoveAt( Index );
-		OnActionsOrderChanged.Broadcast();
-		return true;
-	}
-
-	UE_LOG( LogDataprepCore, Error, TEXT("DataprepAsset::RemoveAction: the Index is out of range") );
-	return false;
-}
-
-void UDataprepAsset::RemoveInvalidActions()
-{
-	bool bWasActionsModified = false;
-	for ( int32 i = 0; i < Actions.Num(); i++ )
-	{
-		UDataprepActionAsset* Action = Actions[ i ].GetActionAsset();
-		if ( !Action || Action->IsPendingKill() )
-		{
-			Actions.RemoveAt( i );
-			i--;
-			bWasActionsModified = true;
-		}
-	}
-
-	if ( bWasActionsModified )
-	{
-		OnActionsOrderChanged.Broadcast();
-	}
-}
-
 void UDataprepAsset::RunProducers(const UDataprepContentProducer::ProducerContext& InContext, TArray< TWeakObjectPtr< UObject > >& OutAssets)
 {
 	OutAssets.Empty();
diff --git a/Engine/Plugins/Experimental/DataPrepEditor/Source/DataPrepCore/Public/DataPrepAsset.h b/Engine/Plugins/Experimental/DataPrepEditor/Source/DataPrepCore/Public/DataPrepAsset.h
--- a/Engine/Plugins/Experimental/DataPrepEditor/Source/DataPrepCore/Public/DataPrepAsset.h
+++ b/Engine/Plugins/Experimental/DataPrepEditor/Source/DataPrepCore/Public/DataPrepAsset.h
@@ -49,86 +49,6 @@ struct FDataprepAssetProducer
 	int32 SupersededBy;
 };
 
-/**
- * Used for the internals of the DataprepAsset only.
- * FDataprepAssetAction manage the state of an action for the DataprepAsset.
- * Also, it report the notification from the action to Dataprep asset.
- */
-USTRUCT(Experimental)
-struct FDataprepAssetAction
-{
-	GENERATED_BODY()
-
-	FDataprepAssetAction()
-		: ActionAsset(nullptr)
-		, bIsEnabled(true)
-		, DataprepAssetPtr(nullptr)
-		, OnOperationOrderChangedHandle()
-	{}
-
-	FDataprepAssetAction(UDataprepActionAsset* InActionAsset, bool bInEnabled, UDataprepAsset& DataprepAsset)
-		: ActionAsset(nullptr)
-		, bIsEnabled(bInEnabled)
-		, DataprepAssetPtr(&DataprepAsset)
-		, OnOperationOrderChangedHandle()
-	{
-		SetActionAsset(InActionAsset);
-	}
-
-	FDataprepAssetAction(const FDataprepAssetAction& Other)
-		: ActionAsset(nullptr)
-		, bIsEnabled(Other.bIsEnabled)
-		, DataprepAssetPtr(Other.DataprepAssetPtr)
-		, OnOperationOrderChangedHandle()
-	{
-		SetActionAsset(Other.ActionAsset);
-	}
-
-	FDataprepAssetAction(FDataprepAssetAction&& Other)
-		: ActionAsset(nullptr)
-		, bIsEnabled(Other.bIsEnabled)
-		, DataprepAssetPtr(Other.DataprepAssetPtr)
-		, OnOperationOrderChangedHandle()
-	{
-		SetActionAsset(Other.ActionAsset);
-		Other.SetActionAsset( nullptr );
-	}
-
-	virtual ~FDataprepAssetAction();
-
-	FDataprepAssetAction& operator=(const FDataprepAssetAction& Other);
-
-	FDataprepAssetAction& operator=(FDataprepAssetAction&& Other);
-
-	UDataprepActionAsset* GetActionAsset() { return ActionAsset; }
-
-	const UDataprepActionAsset* GetActionAsset() const { return ActionAsset; }
-
-	void Enable(bool bInIsEnable) { bIsEnabled = bInIsEnable; }
-
-	bool IsEnabled() const { return bIsEnabled; }
-
-private:
-
-	void SetActionAsset(UDataprepActionAsset* InActionAsset);
-
-	void BindDataprepAssetToAction();
-
-	void UnbindDataprepAssetFromAction();
-
-	void OnActionOperationsOrderChanged();
-
-	UPROPERTY()
-	UDataprepActionAsset* ActionAsset;
-
-	UPROPERTY()
-	bool bIsEnabled;
-
-	TWeakObjectPtr<UDataprepAsset> DataprepAssetPtr;
-
-	FDelegateHandle OnOperationOrderChangedHandle;
-};
-
 
 UCLASS(Experimental)
 class DATAPREPCORE_API UDataprepAsset : public UObject
@@ -138,89 +58,13 @@ class DATAPREPCORE_API UDataprepAsset : public UObject
 public:
 	UDataprepAsset();
 
-	virtual ~UDataprepAsset();
+	virtual ~UDataprepAsset() = default;
 
 	// UObject interface
 	virtual void Serialize( FArchive& Ar ) override;
 	virtual void PostInitProperties() override;
 	// End of UObject interface
 
-	/**
-	 * Create an action at the end of actions array
-	 * @return The index of the added action
-	 */
-	int32 AddAction();
-
-	/**
-	 * Access to an action
-	 * @param Index the index of the desired action
-	 * @return A pointer to the action if it exist, otherwise nullptr
-	 */
-	inline UDataprepActionAsset* GetAction(int32 Index);
-
-	/**
-	 * Access to an action
-	 * @param Index the index of the desired action
-	 * @return A pointer to the action if it exist, otherwise nullptr
-	 */
-	inline const UDataprepActionAsset* GetAction(int32 Index) const;
-
-	/**
-	 * Get the number of action of this Dataprep
-	 * @return The number of action
-	 */
-	inline int32 GetActionsCount() const;
-
-	/**
-	 * Get enabled status of an action
-	 * @param Index The index of the action
-	 * @return True if the action is enabled. Always return false if the action index is invalid
-	 */
-	inline bool IsActionEnabled(int32 Index) const;
-
-	/**
-	 * Allow to set the enabled state of an action
-	 * @param Index The index of the action
-	 * @param bEnable The new enabled state of the action
-	 */
-	void EnableAction(int32 Index, bool bEnable);
-
-	/**
-	 * Move an action to another spot in the order of actions
-	 * This operation take O(n) time. Where n is the absolute value of ActionIndex - DestinationIndex
-	 * @param OperationIndex The Index of the operation to move
-	 * @param DestinationIndex The index of where the action will be move to
-	 * @return True if the operation was move
-	 */
-	bool MoveAction(int32 ActionIndex, int32 DestinationIndex);
-
-	/**
-	 * Remove an action from the Dataprep
-	 * @param Index The index of the action to remove
-	 * @return True if an action was removed
-	 */
-	bool RemoveAction(int32 Index);
-
-	/**
-	 * Allow an observer to be notified when the actions order changed that also include adding and removing an action
-	 * @return The delegate that will be broadcasted when the actions order changed
-	 */
-	DECLARE_EVENT(UDataprepAsset, FOnActionsOrderChanged)
-	FOnActionsOrderChanged& GetOnActionsOrderChanged()
-	{
-		return OnActionsOrderChanged;
-	}
-
-	/**
-	 * Allow an observer to be notified when the operations order of an action changed that also include adding and removing an operation
-	 * @return The delegate that will be broadcasted when the operations order of an action has changed
-	 */
-	DECLARE_EVENT_OneParam(UDataprepAsset, FOnActionOperationsOrderChanged, const UDataprepActionAsset* /** InActionAsset The modified action */)
-	FOnActionOperationsOrderChanged& GetOnActionOperationsOrderChanged()
-	{
-		return OnActionOperationsOrderChanged;
-	}
-
 	/**
 	 * Add a producer of a given class
 	 * @param ProducerClass Class of the Producer to add
@@ -344,18 +188,6 @@ private:
 	  */
 	void ValidateProducerChanges( int32 Index, bool &bChangeAll );
 
-	void RemoveInvalidActions();
-
-	/** Ordered array of the actions */
-	UPROPERTY()
-	TArray<FDataprepAssetAction> Actions;
-
-	/** Delegate broadcasted when the actions order changed */
-	FOnActionsOrderChanged OnActionsOrderChanged;
-
-	/** Delegate broadcasted when the order of the operations of an action changed */
-	FOnActionOperationsOrderChanged OnActionOperationsOrderChanged;
-
 	/** Delegate broadcasted when the consumer or one of the producers has changed */
 	FOnDataprepAssetChanged OnChanged;
 
diff --git a/Engine/Plugins/Experimental/DataPrepEditor/Source/DataPrepCore/Public/DataprepActionAsset.h b/Engine/Plugins/Experimental/DataPrepEditor/Source/DataPrepCore/Public/DataprepActionAsset.h
--- a/Engine/Plugins/Experimental/DataPrepEditor/Source/DataPrepCore/Public/DataprepActionAsset.h
+++ b/Engine/Plugins/Experimental/DataPrepEditor/Source/DataPrepCore/Public/DataprepActionAsset.h
@@ -84,27 +84,27 @@ public:
 	 * @param Index the index of the desired step
 	 * @return A pointer to the step if it exist, otherwise nullptr
 	 */
-	inline TWeakObjectPtr<UDataprepActionStep> GetStep(int32 Index);
+	TWeakObjectPtr<UDataprepActionStep> GetStep(int32 Index);
 
 	/**
 	 * Access to a step of the action
 	 * @param Index the index of the desired step
 	 * @return A const pointer to the operation if it exist, otherwise nullptr
 	 */
-	inline const TWeakObjectPtr<UDataprepActionStep> GetStep(int32 Index) const;
+	const TWeakObjectPtr<UDataprepActionStep> GetStep(int32 Index) const;
 
 	/**
 	 * Get the number of steps of this action 
 	 * @return The number of steps
 	 */
-	inline int32 GetStepsCount() const;
+	int32 GetStepsCount() const;
 
 	/**
 	 * Get enabled status of an operation
 	 * @param Index The index of the operation
 	 * @return True if the operation is enabled. Always return false if the operation index is invalid
 	 */
-	inline bool IsStepEnabled(int32 Index) const;
+	bool IsStepEnabled(int32 Index) const;
 
 	/**
 	 * Allow to set the enabled state of a step
diff --git a/Engine/Plugins/Experimental/DataPrepEditor/Source/DataPrepCore/Public/SelectionSystem/DataprepBoolFilter.h b/Engine/Plugins/Experimental/DataPrepEditor/Source/DataPrepCore/Public/SelectionSystem/DataprepBoolFilter.h
--- a/Engine/Plugins/Experimental/DataPrepEditor/Source/DataPrepCore/Public/SelectionSystem/DataprepBoolFilter.h
+++ b/Engine/Plugins/Experimental/DataPrepEditor/Source/DataPrepCore/Public/SelectionSystem/DataprepBoolFilter.h
@@ -18,7 +18,7 @@ class DATAPREPCORE_API UDataprepBoolFilter : public UDataprepFilter
 	GENERATED_BODY()
 
 public:
-	inline bool Filter(const bool bResult) const;
+	bool Filter(const bool bResult) const;
 
 	//~ Begin UDataprepFilter Interface
 	virtual TArray<UObject*> FilterObjects(const TArray<UObject*>& Objects) const override;
diff --git a/Engine/Plugins/Experimental/DataPrepEditor/Source/DataPrepCore/Public/SelectionSystem/DataprepFloatFilter.h b/Engine/Plugins/Experimental/DataPrepEditor/Source/DataPrepCore/Public/SelectionSystem/DataprepFloatFilter.h
--- a/Engine/Plugins/Experimental/DataPrepEditor/Source/DataPrepCore/Public/SelectionSystem/DataprepFloatFilter.h
+++ b/Engine/Plugins/Experimental/DataPrepEditor/Source/DataPrepCore/Public/SelectionSystem/DataprepFloatFilter.h
@@ -26,7 +26,7 @@ class DATAPREPCORE_API UDataprepFloatFilter : public UDataprepFilter
 	GENERATED_BODY()
 
 public:
-	inline bool Filter(float Float) const;
+	bool Filter(float Float) const;
 
 	//~ Begin UDataprepFilter Interface
 	virtual TArray<UObject*> FilterObjects(const TArray<UObject*>& Objects) const override;
@@ -37,9 +37,9 @@ public:
 	virtual UDataprepFetcher* GetFetcher() const override;
 	//~ Begin UDataprepFilter Interface
 
-	inline EDataprepFloatMatchType GetFloatMatchingCriteria() const;
-	inline float GetEqualValue() const;
-	inline float GetTolerance() const;
+	EDataprepFloatMatchType GetFloatMatchingCriteria() const;
+	float GetEqualValue() const;
+	float GetTolerance() const;
 
 	void SetFloatMatchingCriteria(EDataprepFloatMatchType FloatMatchingCriteria);
 	void SetEqualValue(float EqualValue);
diff --git a/Engine/Plugins/Experimental/DataPrepEditor/Source/DataPrepCore/Public/SelectionSystem/DataprepStringFilter.h b/Engine/Plugins/Experimental/DataPrepEditor/Source/DataPrepCore/Public/SelectionSystem/DataprepStringFilter.h
--- a/Engine/Plugins/Experimental/DataPrepEditor/Source/DataPrepCore/Public/SelectionSystem/DataprepStringFilter.h
+++ b/Engine/Plugins/Experimental/DataPrepEditor/Source/DataPrepCore/Public/SelectionSystem/DataprepStringFilter.h
@@ -27,7 +27,7 @@ class DATAPREPCORE_API UDataprepStringFilter : public UDataprepFilter
 	GENERATED_BODY()
 
 public:
-	inline bool Filter(const FString& String) const;
+	bool Filter(const FString& String) const;
 
 	//~ Begin UDataprepFilter Interface
 	virtual TArray<UObject*> FilterObjects(const TArray<UObject*>& Objects) const override;
@@ -38,8 +38,8 @@ public:
 	virtual UDataprepFetcher* GetFetcher() const override;
 	//~ Begin UDataprepFilter Interface
 
-	inline EDataprepStringMatchType GetStringMatchingCriteria() const;
-	inline FString GetUserString() const;
+	EDataprepStringMatchType GetStringMatchingCriteria() const;
+	FString GetUserString() const;
 
 	void SetStringMatchingCriteria(EDataprepStringMatchType StringMatchingCriteria);
 	void SetUserString(FString UserString);
