Fixed issues with the assets not being properly deleted at the end of a dataprep or when using the operation remove objects.

diff --git a/Engine/Plugins/Experimental/DataPrepEditor/Source/DataPrepCore/Private/DataPrepAsset.cpp b/Engine/Plugins/Experimental/DataPrepEditor/Source/DataPrepCore/Private/DataPrepAsset.cpp
--- a/Engine/Plugins/Experimental/DataPrepEditor/Source/DataPrepCore/Private/DataPrepAsset.cpp
+++ b/Engine/Plugins/Experimental/DataPrepEditor/Source/DataPrepCore/Private/DataPrepAsset.cpp
@@ -4,7 +4,6 @@
 
 #include "DataprepActionAsset.h"
 #include "DataprepCoreLogCategory.h"
-#include "DataprepCoreUtils.h"
 #include "DataPrepRecipe.h"
 
 #include "AssetRegistryModule.h"
diff --git a/Engine/Plugins/Experimental/DataPrepEditor/Source/DataPrepCore/Private/DataprepActionAsset.cpp b/Engine/Plugins/Experimental/DataPrepEditor/Source/DataPrepCore/Private/DataprepActionAsset.cpp
--- a/Engine/Plugins/Experimental/DataPrepEditor/Source/DataPrepCore/Private/DataprepActionAsset.cpp
+++ b/Engine/Plugins/Experimental/DataPrepEditor/Source/DataPrepCore/Private/DataprepActionAsset.cpp
@@ -5,7 +5,7 @@
 // Dataprep include
 #include "DataPrepOperation.h"
 #include "DataprepCoreLogCategory.h"
-#include "DataprepCoreUtils.h"
+#include "DataprepCorePrivateUtils.h"
 #include "SelectionSystem/DataprepFilter.h"
 
 // Engine include
@@ -182,7 +182,7 @@ bool UDataprepActionAsset::MoveStep(int32 StepIndex, int32 DestinationIndex)
 		Modify();
 	}
 
-	if ( DataprepCoreUtils::MoveArrayElement( Steps, StepIndex, DestinationIndex ) )
+	if ( DataprepCorePrivateUtils::MoveArrayElement( Steps, StepIndex, DestinationIndex ) )
 	{
 		OnStepsChanged.Broadcast();
 		return true;
diff --git a/Engine/Plugins/Experimental/DataPrepEditor/Source/DataPrepCore/Private/DataprepCorePrivateUtils.h b/Engine/Plugins/Experimental/DataPrepEditor/Source/DataPrepCore/Private/DataprepCorePrivateUtils.h
new file mode 100644
--- /dev/null
+++ b/Engine/Plugins/Experimental/DataPrepEditor/Source/DataPrepCore/Private/DataprepCorePrivateUtils.h
@@ -0,0 +1,40 @@
+// Copyright 1998-2019 Epic Games, Inc. All Rights Reserved.
+
+#pragma once
+
+#include "CoreMinimal.h"
+
+namespace DataprepCorePrivateUtils
+{
+	/**
+	 * Move an array element to another spot
+	 * This operation take O(n) time. Where n is the absolute value of SourceIndex - DestinationIndex
+	 * @param SourceIndex The Index of the element to move
+	 * @param DestinationIndex The index of where the element will be move to
+	 * @return True if the element was move
+	 */
+	template<class TArrayClass>
+	bool MoveArrayElement(TArrayClass& Array, int32 SourceIndex, int32 DestinationIndex)
+	{
+		if ( Array.IsValidIndex( SourceIndex ) && Array.IsValidIndex( DestinationIndex ) && SourceIndex != DestinationIndex )
+		{
+			// Swap the operation until all operation are at right position. O(n) where n is Upper - Lower 
+
+			while ( SourceIndex < DestinationIndex )
+			{
+				Array.Swap( SourceIndex, DestinationIndex );
+				DestinationIndex--;
+			}
+			
+			while ( DestinationIndex < SourceIndex )
+			{
+				Array.Swap( DestinationIndex, SourceIndex );
+				DestinationIndex++;
+			}
+
+			return true;
+		}
+
+		return false;
+	}
+}
diff --git a/Engine/Plugins/Experimental/DataPrepEditor/Source/DataPrepCore/Private/DataprepCoreUtils.cpp b/Engine/Plugins/Experimental/DataPrepEditor/Source/DataPrepCore/Private/DataprepCoreUtils.cpp
new file mode 100644
--- /dev/null
+++ b/Engine/Plugins/Experimental/DataPrepEditor/Source/DataPrepCore/Private/DataprepCoreUtils.cpp
@@ -0,0 +1,74 @@
+// Copyright 1998-2019 Epic Games, Inc. All Rights Reserved.
+
+#include "DataprepCoreUtils.h"
+
+#include "UObject/UObjectHash.h"
+
+#if WITH_EDITOR
+#include "ObjectTools.h"
+#include "Toolkits/AssetEditorManager.h"
+#endif
+
+void FDataprepCoreUtils::PurgeObjects(TArray<UObject*> Objects)
+{
+	TArray<UObject*> ObjectsToPurge;
+	ObjectsToPurge.Reserve( Objects.Num() );
+	TArray<UObject*> PublicObjectsToPurge;
+	PublicObjectsToPurge.Reserve( Objects.Num() );
+
+	auto MakeObjectPurgeable = [&ObjectsToPurge, &PublicObjectsToPurge](UObject* InObject)
+	{
+#if WITH_EDITOR
+		if ( InObject->IsAsset() )
+		{
+			FAssetEditorManager::Get().CloseAllEditorsForAsset( InObject );
+		}
+#endif
+		if ( InObject->IsRooted() )
+		{
+			InObject->RemoveFromRoot();
+		}
+
+		if ( InObject->HasAnyFlags( RF_Public ) )
+		{
+			PublicObjectsToPurge.Add(InObject);
+		}
+
+		InObject->ClearFlags( RF_Public | RF_Standalone );
+		InObject->MarkPendingKill();
+		ObjectsToPurge.Add( InObject );
+	};
+
+	auto MakeSourceObjectPurgeable = [MakeObjectPurgeable](UObject* InSourceObject)
+	{
+		MakeObjectPurgeable( InSourceObject );
+		ForEachObjectWithOuter( InSourceObject, [MakeObjectPurgeable](UObject* InObject)
+		{
+			MakeObjectPurgeable( InObject );
+		});
+	};
+
+	// Clean-up any in-memory packages that should be purged and check if we are purging the current map
+	for ( UObject* Object : Objects )
+	{
+		if (Object)
+		{
+			MakeSourceObjectPurgeable(Object);
+		}
+	}
+
+	// If we have any object that were made purgeable, null out their references so we can garbage collect
+	if ( PublicObjectsToPurge.Num() > 0 )
+	{
+#if WITH_EDITOR
+		ObjectTools::ForceReplaceReferences(nullptr, PublicObjectsToPurge);
+#endif // WITH_EDITOR
+	}
+
+	// if we have object to purge but the map isn't one of them collect garbage (if we purged the map it has already been done)
+	if ( ObjectsToPurge.Num() > 0 )
+	{
+		CollectGarbage(GARBAGE_COLLECTION_KEEPFLAGS);
+	}
+}
+
diff --git a/Engine/Plugins/Experimental/DataPrepEditor/Source/DataPrepCore/Private/DataprepCoreUtils.h b/Engine/Plugins/Experimental/DataPrepEditor/Source/DataPrepCore/Private/DataprepCoreUtils.h
deleted file mode 100644
diff --git a/Engine/Plugins/Experimental/DataPrepEditor/Source/DataPrepCore/Public/DataPrepAsset.h b/Engine/Plugins/Experimental/DataPrepEditor/Source/DataPrepCore/Public/DataPrepAsset.h
--- a/Engine/Plugins/Experimental/DataPrepEditor/Source/DataPrepCore/Public/DataPrepAsset.h
+++ b/Engine/Plugins/Experimental/DataPrepEditor/Source/DataPrepCore/Public/DataPrepAsset.h
@@ -190,6 +190,4 @@ private:
 
 	/** Delegate broadcasted when the consumer or one of the producers has changed */
 	FOnDataprepAssetChanged OnChanged;
-
-	FDelegateHandle OnAssetDeletedHandle;
 };
\ No newline at end of file
diff --git a/Engine/Plugins/Experimental/DataPrepEditor/Source/DataPrepCore/Public/DataprepCoreUtils.h b/Engine/Plugins/Experimental/DataPrepEditor/Source/DataPrepCore/Public/DataprepCoreUtils.h
new file mode 100644
--- /dev/null
+++ b/Engine/Plugins/Experimental/DataPrepEditor/Source/DataPrepCore/Public/DataprepCoreUtils.h
@@ -0,0 +1,12 @@
+// Copyright 1998-2019 Epic Games, Inc. All Rights Reserved.
+
+#pragma once
+
+#include "CoreMinimal.h"
+
+class DATAPREPCORE_API FDataprepCoreUtils
+{
+public:
+	static void PurgeObjects(TArray<UObject*> Objects);
+};
+
diff --git a/Engine/Plugins/Experimental/DataPrepEditor/Source/DataPrepEditor/Private/DataPrepEditor.cpp b/Engine/Plugins/Experimental/DataPrepEditor/Source/DataPrepEditor/Private/DataPrepEditor.cpp
--- a/Engine/Plugins/Experimental/DataPrepEditor/Source/DataPrepEditor/Private/DataPrepEditor.cpp
+++ b/Engine/Plugins/Experimental/DataPrepEditor/Source/DataPrepEditor/Private/DataPrepEditor.cpp
@@ -11,11 +11,12 @@
 #include "DataPrepEditorStyle.h"
 #include "DataPrepRecipe.h"
 #include "DataprepActionAsset.h"
+#include "DataprepCoreUtils.h"
 #include "DataprepEditorLogCategory.h"
 #include "SchemaActions/DataprepOperationMenuActionCollector.h"
 #include "Widgets/DataprepAssetView.h"
-#include "Widgets/SDataprepPalette.h"
 #include "Widgets/SAssetsPreviewWidget.h"
+#include "Widgets/SDataprepPalette.h"
 
 #include "ActorEditorUtils.h"
 #include "AssetDeleteModel.h"
@@ -596,15 +597,11 @@ void FDataprepEditor::CleanPreviewWorld()
 		}
 	}
 
-	// Prevent display of progress dialog from ObjectTools::DeleteObjectsUnchecked
-	TGuardValue<bool>( GIsSilent, true );
-
-	ensure( ObjectTools::DeleteObjectsUnchecked( ObjectsToDelete ) == ObjectsToDelete.Num() );
+	FDataprepCoreUtils::PurgeObjects( MoveTemp( ObjectsToDelete ) );
 
 	CachedAssets.Reset();
 	Assets.Reset();
 
-	CollectGarbage( GARBAGE_COLLECTION_KEEPFLAGS );
 	PreviewWorld->CleanupActors();
 }
 
@@ -683,7 +680,7 @@ void FDataprepEditor::OnExecutePipeline()
 					while ( Index < Assets.Num() )
 					{
 						UObject* Object = Assets[Index].Get();
-						if ( Object && Object->IsValidLowLevelFast() )
+						if ( Object && Object->IsValidLowLevel() )
 						{
 							Index++;
 						}
diff --git a/Engine/Plugins/Experimental/DataPrepEditor/Source/DataPrepLibraries/Private/DataPrepOperationsLibrary.cpp b/Engine/Plugins/Experimental/DataPrepEditor/Source/DataPrepLibraries/Private/DataPrepOperationsLibrary.cpp
--- a/Engine/Plugins/Experimental/DataPrepEditor/Source/DataPrepLibraries/Private/DataPrepOperationsLibrary.cpp
+++ b/Engine/Plugins/Experimental/DataPrepEditor/Source/DataPrepLibraries/Private/DataPrepOperationsLibrary.cpp
@@ -2,6 +2,8 @@
 
 #include "DataPrepOperationsLibrary.h"
 
+#include "DataprepCoreUtils.h"
+
 #include "ActorEditorUtils.h"
 #include "AssetDeleteModel.h"
 #include "AssetRegistryModule.h"
@@ -21,7 +23,6 @@
 #include "Misc/FileHelper.h"
 #include "ObjectTools.h"
 #include "UObject/SoftObjectPath.h"
-#include "AssetDeleteModel.h"
 
 DEFINE_LOG_CATEGORY(LogDataprep);
 
@@ -477,17 +478,8 @@ void UDataprepOperationsLibrary::RemoveObjects(const TArray< UObject* >& Objects
 		GEditor->NoteSelectionChange();
 	}
 
-	// Delete assets
-	{
-		const int32 AssetsToDeleteCount = AssetsToDelete.Num();
-		const int32 NonDeletedAssetsCount = ObjectTools::DeleteObjects(AssetsToDelete, false) - AssetsToDeleteCount;
+	FDataprepCoreUtils::PurgeObjects( MoveTemp( AssetsToDelete ) );
 
-		// Warn user that not all objects have been deleted
-		if(NonDeletedAssetsCount > 0)
-		{
-			UE_LOG(LogDataprep, Warning, TEXT("RemoveObjects: %d objects out of %d were not deleted"), NonDeletedAssetsCount, AssetsToDeleteCount);
-		}
-	}
 }
 
 #undef LOCTEXT_NAMESPACE
