Fxing the set generate lightmap not leaving a proper lighmap channel.

diff --git a/Engine/Plugins/Experimental/DataPrepEditor/Source/DataPrepLibraries/Private/DataPrepOperationsLibrary.cpp b/Engine/Plugins/Experimental/DataPrepEditor/Source/DataPrepLibraries/Private/DataPrepOperationsLibrary.cpp
--- a/Engine/Plugins/Experimental/DataPrepEditor/Source/DataPrepLibraries/Private/DataPrepOperationsLibrary.cpp
+++ b/Engine/Plugins/Experimental/DataPrepEditor/Source/DataPrepLibraries/Private/DataPrepOperationsLibrary.cpp
@@ -6,6 +6,7 @@
 
 #include "ActorEditorUtils.h"
 #include "AssetDeleteModel.h"
+#include "AssetDeleteModel.h"
 #include "AssetRegistryModule.h"
 #include "Camera/CameraActor.h"
 #include "Editor.h"
@@ -19,6 +20,9 @@
 #include "Layers/ILayers.h"
 #include "Materials/Material.h"
 #include "Materials/MaterialInterface.h"
+#include "Math/Vector2D.h"
+#include "MeshAttributeArray.h"
+#include "MeshDescription.h"
 #include "MeshTypes.h"
 #include "Misc/FileHelper.h"
 #include "ObjectTools.h"
@@ -245,9 +249,22 @@ void UDataprepOperationsLibrary::SetGenerateLightmapUVs( const TArray< UObject*
 	{
 		if (StaticMesh)
 		{
+			// 3 is the maximum that lightmass accept
+			int32 MinBiggestUVChannel = 3;
 			for ( FStaticMeshSourceModel& SourceModel : StaticMesh->SourceModels )
 			{
 				SourceModel.BuildSettings.bGenerateLightmapUVs = bGenerateLightmapUVs;
+				if( FMeshDescription* MeshDescription = SourceModel.MeshDescription.Get() )
+				{
+					int32 UVChannelCount = MeshDescription->VertexInstanceAttributes().GetAttributesRef< FVector2D >( MeshAttribute::VertexInstance::TextureCoordinate ).GetNumIndices();
+					MinBiggestUVChannel = FMath::Min( MinBiggestUVChannel, UVChannelCount - 1 );
+				}
+			}
+
+			if ( StaticMesh->LightMapCoordinateIndex > MinBiggestUVChannel )
+			{
+				// Correct the coordinate index if it was invalid
+				StaticMesh->LightMapCoordinateIndex = MinBiggestUVChannel;
 			}
 		}
 	}
