Fxing an issue with my previous commit

diff --git a/Engine/Plugins/Experimental/DataPrepEditor/Source/DataPrepLibraries/Private/DataPrepOperationsLibrary.cpp b/Engine/Plugins/Experimental/DataPrepEditor/Source/DataPrepLibraries/Private/DataPrepOperationsLibrary.cpp
--- a/Engine/Plugins/Experimental/DataPrepEditor/Source/DataPrepLibraries/Private/DataPrepOperationsLibrary.cpp
+++ b/Engine/Plugins/Experimental/DataPrepEditor/Source/DataPrepLibraries/Private/DataPrepOperationsLibrary.cpp
@@ -249,10 +249,13 @@ void UDataprepOperationsLibrary::SetGenerateLightmapUVs( const TArray< UObject*
 	{
 		if (StaticMesh)
 		{
+			bool bDidChangeSettings = false;
+
 			// 3 is the maximum that lightmass accept
 			int32 MinBiggestUVChannel = 3;
 			for ( FStaticMeshSourceModel& SourceModel : StaticMesh->SourceModels )
 			{
+				bDidChangeSettings |= SourceModel.BuildSettings.bGenerateLightmapUVs != bGenerateLightmapUVs;
 				SourceModel.BuildSettings.bGenerateLightmapUVs = bGenerateLightmapUVs;
 				if( FMeshDescription* MeshDescription = SourceModel.MeshDescription.Get() )
 				{
@@ -261,7 +264,7 @@ void UDataprepOperationsLibrary::SetGenerateLightmapUVs( const TArray< UObject*
 				}
 			}
 
-			if ( StaticMesh->LightMapCoordinateIndex > MinBiggestUVChannel )
+			if ( StaticMesh->LightMapCoordinateIndex > MinBiggestUVChannel && bDidChangeSettings )
 			{
 				// Correct the coordinate index if it was invalid
 				StaticMesh->LightMapCoordinateIndex = MinBiggestUVChannel;
