- Added automatic opening of explorer when adding a Datasmith producer, file or directory
- Changed logic to only show combo box with consumers if there are more than one consumer defined
- Enabled snapshot based execution
- Changed Level to Sub-Level in output widget

diff --git a/Engine/Plugins/Experimental/DataPrepEditor/Source/DataPrepEditor/Private/DataprepSnapshot.cpp b/Engine/Plugins/Experimental/DataPrepEditor/Source/DataPrepEditor/Private/DataprepSnapshot.cpp
--- a/Engine/Plugins/Experimental/DataPrepEditor/Source/DataPrepEditor/Private/DataprepSnapshot.cpp
+++ b/Engine/Plugins/Experimental/DataPrepEditor/Source/DataPrepEditor/Private/DataprepSnapshot.cpp
@@ -13,7 +13,9 @@
 #include "Factories/LevelFactory.h"
 #include "HAL/FileManager.h"
 #include "GenericPlatform/GenericPlatformOutputDevices.h"
+#include "GenericPlatform/GenericPlatformTime.h"
 #include "Materials/MaterialInstance.h"
+#include "Misc/Compression.h"
 #include "Misc/FileHelper.h"
 #include "ObjectTools.h"
 #include "PackageTools.h"
@@ -37,7 +39,8 @@ enum class EDataprepAssetClass : uint8 {
 };
 
 // #ueent_todo: Boolean driving activating actual snapshot based logic
-const bool bUseSnapshot = false;
+const bool bUseSnapshot = true;
+const bool bUseCompression = true;
 
 namespace DataprepSnapshotUtil
 {
@@ -63,31 +66,38 @@ namespace DataprepSnapshotUtil
 		return FPaths::ConvertRelativePathToFull( FPaths::Combine( RootPath, PackageFileName ) + SnapshotExtension );
 	}
 
-	struct FDataprepSnapshotParser : public FArchiveUObject
+	void WriteSnapshotData(UObject* Object, TArray<uint8>& OutSerializedData, TMap<FString, UClass*>& OutClassesMap)
 	{
-		virtual FArchive& operator<<(UObject*& Obj) override
+		// Helper struct to identify dependency of a UObject on other UObject(s) except given one (its outer)
+		struct FObjectDependencyAnalyzer : public FArchiveUObject
 		{
-			if(Obj != nullptr)
+			FObjectDependencyAnalyzer(UObject* InObject)
+				: ObjectToSkip(InObject)
+			{ }
+
+			virtual FArchive& operator<<(UObject*& Obj) override
 			{
-				if( Obj->IsA<UPackage>() || Obj->HasAnyFlags(RF_Public) )
-				{
-					return FArchiveUObject::operator<<(Obj);
-				}
-				else if(SerializedObjects.Find(Obj) == nullptr)
+				if(Obj != nullptr)
 				{
-					SerializedObjects.Add(Obj);
+					if( Obj->IsA<UPackage>() || (Obj->HasAnyFlags(RF_Public) && Obj->GetOuter()->IsA<UPackage>()))
+					{
+						return FArchiveUObject::operator<<(Obj);
+					}
+					else if(Obj != ObjectToSkip && DependentObjects.Find(Obj) == nullptr)
+					{
+						DependentObjects.Add(Obj);
+					}
 				}
-			}
 
-			return *this;
-		}
+				return *this;
+			}
 
-		TSet<UObject*> SerializedObjects;
-	};
+			TSet<UObject*> DependentObjects;
+			UObject* ObjectToSkip;
+		};
 
-	void WriteSnapshotData(UObject* Object, TArray<uint8>& OutSerializedData, TMap<FString, UClass*>& OutClassesMap)
-	{
-		FMemoryWriter MemAr(OutSerializedData);
+		TArray<uint8> MemoryBuffer;
+		FMemoryWriter MemAr(MemoryBuffer);
 		FObjectAndNameAsStringProxyArchive Ar(MemAr, false);
 		Ar.SetIsTransacting(true);
 
@@ -95,7 +105,9 @@ namespace DataprepSnapshotUtil
 		TArray< UObject* > SubObjectsArray;
 		GetObjectsWithOuter( Object, SubObjectsArray, /*bIncludeNestedObjects = */ true );
 
+		// Sort array of sub-objects based on their inter-dependency
 		{
+			// Create and initialize graph of dependency between sub-objects
 			TMap< UObject*, TSet<UObject*> > SubObjectDependencyGraph;
 			SubObjectDependencyGraph.Reserve( SubObjectsArray.Num() );
 
@@ -104,20 +116,17 @@ namespace DataprepSnapshotUtil
 				SubObjectDependencyGraph.Add( SubObject );
 			}
 
+			// Build graph of dependency: each entry contains the set of sub-objects to create before itself
 			for(UObject* SubObject : SubObjectsArray)
 			{
-				FDataprepSnapshotParser Parser;
-				SubObject->Serialize(Parser);
+				FObjectDependencyAnalyzer Analyzer(SubObject->GetOuter());
+				SubObject->Serialize( Analyzer );
 
-				for(UObject* SubObjectDependence : Parser.SerializedObjects)
-				{
-					if(SubObjectDependence != SubObject->GetOuter())
-					{
-						SubObjectDependencyGraph[SubObject].Add( SubObjectDependence );
-					}
-				}
+				SubObjectDependencyGraph[SubObject].Append(Analyzer.DependentObjects);
 			}
 
+			// Sort array of sub-objects: first objects do not depend on ones below
+			// #ueent_todo: Improve performance of building. Current is pretty brute force
 			int32 Count = SubObjectsArray.Num();
 			SubObjectsArray.Empty(Count);
 
@@ -147,9 +156,11 @@ namespace DataprepSnapshotUtil
 			}
 		}
 
-		int32 Count = SubObjectsArray.Num();
-		MemAr << Count;
+		// Serialize size of array
+		int32 SubObjectsCount = SubObjectsArray.Num();
+		MemAr << SubObjectsCount;
 
+		// Serialize class and path of each sub-object
 		for(UObject* SubObject : SubObjectsArray)
 		{
 			UClass* SubObjectClass = SubObject->GetClass();
@@ -158,32 +169,78 @@ namespace DataprepSnapshotUtil
 			MemAr << ClassName;
 
 			OutClassesMap.Add(ClassName, SubObjectClass);
+		}
+
+		// Serialize sub-objects' outer path
+		for(UObject* SubObject : SubObjectsArray)
+		{
+			FSoftObjectPath SoftPath( SubObject->GetOuter() );
 
-			FString ObjectPath = SubObject->GetPathName();
-			MemAr << ObjectPath;
+			FString SoftPathString = SoftPath.ToString();
+			MemAr << SoftPathString;
 		}
 
+		// Serialize sub-objects' content
 		for(UObject* SubObject : SubObjectsArray)
 		{
 			SubObject->Serialize(Ar);
 		}
 
+		// Serialize object
 		Object->Serialize(Ar);
+
+		if(bUseCompression)
+		{
+			const int32 BufferHeaderSize = (int32)sizeof(int32);
+
+			// allocate all of the input space for the output, with extra space for max overhead of zlib (when compressed > uncompressed)
+			int32 CompressedSize = FCompression::CompressMemoryBound(NAME_Zlib, MemoryBuffer.Num());
+			OutSerializedData.SetNum( CompressedSize + BufferHeaderSize );
+
+			// Store the size of the uncompressed buffer
+			((int32*)OutSerializedData.GetData())[0] = MemoryBuffer.Num();
+
+			// compress the data		
+			bool bSucceeded = FCompression::CompressMemory(NAME_Zlib, OutSerializedData.GetData() + BufferHeaderSize, CompressedSize, MemoryBuffer.GetData(), MemoryBuffer.Num());
+
+			// if it failed send the data uncompressed, which we mark by setting compressed size to 0
+			checkf(bSucceeded, TEXT("zlib failed to compress, which is very unexpected"));
+
+			OutSerializedData.SetNum( CompressedSize + BufferHeaderSize );
+		}
 	}
 
 	void ReadSnapshotData(UObject* Object, const TArray<uint8>& InSerializedData, TMap<FString, UClass*>& InClassesMap)
 	{
-		FMemoryReader MemAr(InSerializedData);
+		TArray<uint8> MemoryBuffer;
+		if(bUseCompression)
+		{
+			const int32 BufferHeaderSize = (int32)sizeof(int32);
+
+			// Allocate the space required for the uncompressed data
+			int32 UncompressedSize = ((int32*)InSerializedData.GetData())[0];
+			MemoryBuffer.SetNum(UncompressedSize);
+
+			// uncompress the data		
+			bool bSucceeded = FCompression::UncompressMemory(NAME_Zlib, MemoryBuffer.GetData(), MemoryBuffer.Num(), InSerializedData.GetData() + BufferHeaderSize, InSerializedData.Num() - BufferHeaderSize);
+
+			// if it failed send the data uncompressed, which we mark by setting compressed size to 0
+			checkf(bSucceeded, TEXT("zlib failed to uncompress, which is very unexpected"));
+		}
+
+		FMemoryReader MemAr(bUseCompression ? MemoryBuffer : InSerializedData);
 		FObjectAndNameAsStringProxyArchive Ar(MemAr, false);
 		Ar.SetIsTransacting(true);
 
-		int32 Count = 0;
-		MemAr << Count;
+		// Deserialize count of sub-objects
+		int32 SubObjectsCount = 0;
+		MemAr << SubObjectsCount;
 
+		// Create empty sub-objects based on class and patch
 		TArray< UObject* > SubObjectsArray;
-		SubObjectsArray.Reserve(Count);
+		SubObjectsArray.Reserve(SubObjectsCount);
 
-		for(int32 Index = 0; Index < Count; ++Index)
+		for(int32 Index = 0; Index < SubObjectsCount; ++Index)
 		{
 			FString ClassName;
 			MemAr << ClassName;
@@ -191,21 +248,35 @@ namespace DataprepSnapshotUtil
 			UClass** SubObjectClassPtr = InClassesMap.Find(ClassName);
 			check( SubObjectClassPtr );
 
-			FString StrObjectPath;
-			MemAr << StrObjectPath;
-
+			// #ueent_todo: TBF: outer will not be properly updated if sub-object depends on another sub-object 
 			UObject* SubObject = NewObject<UObject>( Object, *SubObjectClassPtr, NAME_None, RF_Transient );
 			SubObjectsArray.Add( SubObject );
 		}
 
+		// Restore sub-objects' outer if original outer differs from Object
+		for(UObject* SubObject : SubObjectsArray)
+		{
+			FString SoftPathString;
+			MemAr << SoftPathString;
+
+			FSoftObjectPath SoftPath( SoftPathString );
+			UObject* NewOuter = SoftPath.ResolveObject();
+			ensure( NewOuter );
+
+			if( NewOuter && NewOuter != SubObject->GetOuter() )
+			{
+				SubObject->Rename( nullptr, NewOuter, REN_DontCreateRedirectors | REN_NonTransactional );
+			}
+		}
+
+		// Deserialize sub-objects
 		for(UObject* SubObject : SubObjectsArray)
 		{
 			SubObject->Serialize(Ar);
 		}
 
+		// Deserialize object
 		Object->Serialize(Ar);
-
-		SubObjectsArray.Empty();
 	}
 }
 
@@ -292,6 +363,7 @@ void FDataprepEditor::TakeSnapshot()
 		return;
 	}
 
+	uint64 StartTime = FPlatformTime::Cycles64();
 	UE_LOG( LogDataprepEditor, Log, TEXT("Restoring snapshot...") );
 
 	// Clean up temporary folder with content of previous snapshot(s)
@@ -354,7 +426,7 @@ void FDataprepEditor::TakeSnapshot()
 			FSoftObjectPath AssetPath( AssetObject );
 			ContentSnapshot.DataEntries.Emplace( AssetPath.GetAssetPathString(), AssetObject->GetClass(), AssetObject->GetFlags() );
 
-			UE_LOG( LogDataprepEditor, Log, TEXT("Saving asset %s"), *AssetPath.GetAssetPathString() );
+			UE_LOG( LogDataprepEditor, Verbose, TEXT("Saving asset %s"), *AssetPath.GetAssetPathString() );
 
 			// Serialize asset
 			{
@@ -372,7 +444,7 @@ void FDataprepEditor::TakeSnapshot()
 				break;
 			}
 
-			UE_LOG( LogDataprepEditor, Log, TEXT("Asset %s successfully saved"), *AssetPath.GetAssetPathString() );
+			UE_LOG( LogDataprepEditor, Verbose, TEXT("Asset %s successfully saved"), *AssetPath.GetAssetPathString() );
 		}
 	}
 
@@ -390,7 +462,7 @@ void FDataprepEditor::TakeSnapshot()
 			UExporter::ExportToOutputDevice( &Context, PreviewWorld.Get(), NULL, Ar, TEXT("copy"), 0, ExportFlags);
 
 			// Save text into file
-			FString PackageFilePath = DataprepSnapshotUtil::BuildAssetFileName( TempDir, GetTransientContentFolder() / SessionID );
+			FString PackageFilePath = DataprepSnapshotUtil::BuildAssetFileName( TempDir, GetTransientContentFolder() / SessionID ) + TEXT(".asc");
 			ContentSnapshot.bIsValid &= FFileHelper::SaveStringToFile( Ar, *PackageFilePath );
 		}
 		PreviewWorld->SetFlags(RF_Transient);
@@ -425,6 +497,13 @@ void FDataprepEditor::TakeSnapshot()
 	{
 		return GetAssetClassEnum( A.Get<1>() ) < GetAssetClassEnum( B.Get<1>() );
 	});
+
+	// Log time spent to import incoming file in minutes and seconds
+	double ElapsedSeconds = FPlatformTime::ToSeconds64(FPlatformTime::Cycles64() - StartTime);
+
+	int ElapsedMin = int(ElapsedSeconds / 60.0);
+	ElapsedSeconds -= 60.0 * (double)ElapsedMin;
+	UE_LOG( LogDataprepEditor, Log, TEXT("Snapshot taken in [%d min %.3f s]"), ElapsedMin, ElapsedSeconds );
 }
 
 void FDataprepEditor::RestoreFromSnapshot()
@@ -443,7 +522,8 @@ void FDataprepEditor::RestoreFromSnapshot()
 		return;
 	}
 
-	UE_LOG( LogDataprepEditor, Log, TEXT("Restoring snapshot...") );
+	uint64 StartTime = FPlatformTime::Cycles64();
+	UE_LOG( LogDataprepEditor, Log, TEXT("Cleaning up preview world ...") );
 
 	// Clean up all assets and world content
 	{
@@ -451,6 +531,16 @@ void FDataprepEditor::RestoreFromSnapshot()
 		Assets.Reset(ContentSnapshot.DataEntries.Num());
 	}
 
+	// Log time spent to import incoming file in minutes and seconds
+	double ElapsedSeconds = FPlatformTime::ToSeconds64(FPlatformTime::Cycles64() - StartTime);
+
+	int ElapsedMin = int(ElapsedSeconds / 60.0);
+	ElapsedSeconds -= 60.0 * (double)ElapsedMin;
+	UE_LOG( LogDataprepEditor, Log, TEXT("Preview world cleaned in [%d min %.3f s]"), ElapsedMin, ElapsedSeconds );
+
+	StartTime = FPlatformTime::Cycles64();
+	UE_LOG( LogDataprepEditor, Log, TEXT("Restoring snapshot...") );
+
 	TMap<FString, UPackage*> PackagesCreated;
 	PackagesCreated.Reserve(ContentSnapshot.DataEntries.Num());
 
@@ -465,7 +555,7 @@ void FDataprepEditor::RestoreFromSnapshot()
 		const FString PackageToLoadPath = ObjectPath.GetLongPackageName();
 		const FString AssetName = ObjectPath.GetAssetName();
 
-		UE_LOG( LogDataprepEditor, Log, TEXT("Loading asset %s"), *ObjectPath.GetAssetPathString() );
+		UE_LOG( LogDataprepEditor, Verbose, TEXT("Loading asset %s"), *ObjectPath.GetAssetPathString() );
 
 		if(PackagesCreated.Find(PackageToLoadPath) == nullptr)
 		{
@@ -489,7 +579,7 @@ void FDataprepEditor::RestoreFromSnapshot()
 
 		Assets.Add( Asset );
 
-		UE_LOG( LogDataprepEditor, Log, TEXT("Asset %s loaded"), *ObjectPath.GetAssetPathString() );
+		UE_LOG( LogDataprepEditor, Verbose, TEXT("Asset %s loaded"), *ObjectPath.GetAssetPathString() );
 	}
 
 	UE_LOG( LogDataprepEditor, Log, TEXT("Loading level") );
@@ -497,7 +587,7 @@ void FDataprepEditor::RestoreFromSnapshot()
 		// Code inspired from UUnrealEdEngine::edactPasteSelected
 		ULevel* WorldLevel = PreviewWorld->GetCurrentLevel();
 
-		FString PackageFilePath = DataprepSnapshotUtil::BuildAssetFileName( TempDir, GetTransientContentFolder() / SessionID );
+		FString PackageFilePath = DataprepSnapshotUtil::BuildAssetFileName( TempDir, GetTransientContentFolder() / SessionID ) + TEXT(".asc");
 		const bool bBSPAutoUpdate = GetDefault<ULevelEditorMiscSettings>()->bBSPAutoUpdate;
 		GetMutableDefault<ULevelEditorMiscSettings>()->bBSPAutoUpdate = false;
 
@@ -505,14 +595,21 @@ void FDataprepEditor::RestoreFromSnapshot()
 		FString FileBuffer;
 		check( FFileHelper::LoadFileToString(FileBuffer, *PackageFilePath) );
 
-		// Set the GWorl to the preview world since ULevelFactory::FactoryCreateText uses GWorld
+		// Set the GWorld to the preview world since ULevelFactory::FactoryCreateText uses GWorld
 		UWorld* CachedWorld = GWorld;
 		GWorld = PreviewWorld.Get();
 
+		// Disable warnings from LogExec because ULevelFactory::FactoryCreateText is pretty verbose on harmless warnings
+		ELogVerbosity::Type PrevLogExecVerbosity = LogExec.GetVerbosity();
+		LogExec.SetVerbosity( ELogVerbosity::Error );
+
 		const TCHAR* Paste = *FileBuffer;
 		ULevelFactory* Factory = NewObject<ULevelFactory>();
 		Factory->FactoryCreateText( ULevel::StaticClass(), WorldLevel, WorldLevel->GetFName(), RF_Transactional, NULL, TEXT("paste"), Paste, Paste + FileBuffer.Len(), FGenericPlatformOutputDevices::GetFeedbackContext());
 
+		// Restore LogExec verbosity
+		LogExec.SetVerbosity( PrevLogExecVerbosity );
+
 		// Reinstate old BSP update setting, and force a rebuild - any levels whose geometry has changed while pasting will be rebuilt
 		GetMutableDefault<ULevelEditorMiscSettings>()->bBSPAutoUpdate = bBSPAutoUpdate;
 
@@ -520,6 +617,16 @@ void FDataprepEditor::RestoreFromSnapshot()
 		GWorld = CachedWorld;
 	}
 	UE_LOG( LogDataprepEditor, Log, TEXT("Level loaded") );
+
+	// Log time spent to import incoming file in minutes and seconds
+	ElapsedSeconds = FPlatformTime::ToSeconds64(FPlatformTime::Cycles64() - StartTime);
+
+	ElapsedMin = int(ElapsedSeconds / 60.0);
+	ElapsedSeconds -= 60.0 * (double)ElapsedMin;
+	UE_LOG( LogDataprepEditor, Log, TEXT("Preview world restored in [%d min %.3f s]"), ElapsedMin, ElapsedSeconds );
+
+	// Update preview panels to reflect restored content
+	UpdatePreviewPanels();
 }
 
 #undef LOCTEXT_NAMESPACE
diff --git a/Engine/Plugins/Experimental/DataPrepEditor/Source/DataPrepEditor/Private/Widgets/DataprepAssetView.cpp b/Engine/Plugins/Experimental/DataPrepEditor/Source/DataPrepEditor/Private/Widgets/DataprepAssetView.cpp
--- a/Engine/Plugins/Experimental/DataPrepEditor/Source/DataPrepEditor/Private/Widgets/DataprepAssetView.cpp
+++ b/Engine/Plugins/Experimental/DataPrepEditor/Source/DataPrepEditor/Private/Widgets/DataprepAssetView.cpp
@@ -494,9 +494,23 @@ void SDataprepAssetView::Construct( const FArguments& InArgs, UDataprepAsset* In
 		}
 	}
 
-	if ( !SelectedConsumerDescription.IsValid() && ConsumerDescriptionList.Num() > 0)
+	// Display a combo-box if there are more than one type of consumers
+	if( ConsumerDescriptionMap.Num() > 1 )
 	{
-		SelectedConsumerDescription = MakeShared<FString>( FString() );
+		if ( !SelectedConsumerDescription.IsValid() )
+		{
+			SelectedConsumerDescription = MakeShared<FString>( FString() );
+		}
+
+
+		ProducerSelector = SNew( STextComboBox )
+		.OptionsSource( &ConsumerDescriptionList )
+		.OnSelectionChanged( this, &SDataprepAssetView::OnNewConsumerSelected )
+		.InitiallySelectedItem( SelectedConsumerDescription );
+	}
+	else
+	{
+		ProducerSelector = SNullWidget::NullWidget;
 	}
 
 	TSharedPtr<SWidget> AddNewMenu = SNew(SComboButton)
@@ -617,10 +631,7 @@ void SDataprepAssetView::Construct( const FArguments& InArgs, UDataprepAsset* In
 						.HAlign(EHorizontalAlignment::HAlign_Right)
 						.Padding(0, 0, 2, 0)
 						[
-							SAssignNew( ProducerSelector, STextComboBox )
-							.OptionsSource( &ConsumerDescriptionList )
-							.OnSelectionChanged( this, &SDataprepAssetView::OnNewConsumerSelected )
-							.InitiallySelectedItem( SelectedConsumerDescription )
+							ProducerSelector.ToSharedRef()
 						]
 					]
 					+ SVerticalBox::Slot()
@@ -677,7 +688,7 @@ void SDataprepAssetView::OnNewConsumerSelected( TSharedPtr<FString> NewConsumerD
 
 		if( !DataprepAsset->ReplaceConsumer( *NewConsumerClassPtr ) )
 		{
-			ProducerSelector->SetSelectedItem(SelectedConsumerDescription);
+			((STextComboBox*)ProducerSelector.Get())->SetSelectedItem(SelectedConsumerDescription);
 		}
 		// Update SelectedConsumerDescription only, the widget displaying the consumer is updated thru notifications 
 		else
diff --git a/Engine/Plugins/Experimental/DataPrepEditor/Source/DataPrepEditor/Private/Widgets/DataprepAssetView.h b/Engine/Plugins/Experimental/DataPrepEditor/Source/DataPrepEditor/Private/Widgets/DataprepAssetView.h
--- a/Engine/Plugins/Experimental/DataPrepEditor/Source/DataPrepEditor/Private/Widgets/DataprepAssetView.h
+++ b/Engine/Plugins/Experimental/DataPrepEditor/Source/DataPrepEditor/Private/Widgets/DataprepAssetView.h
@@ -161,7 +161,7 @@ private:
 	TArray< TSharedPtr< FString > > ConsumerDescriptionList;
 	TMap< TSharedPtr< FString >, UClass* > ConsumerDescriptionMap;
 	TSharedPtr< FString > SelectedConsumerDescription;
-	TSharedPtr< STextComboBox > ProducerSelector;
+	TSharedPtr< SWidget > ProducerSelector;
 	bool bIsChecked;
 	TSharedPtr< FProducerStackEntry > SelectedEntry;
 	TSharedPtr< SDataprepConsumerWidget > ConsumerWidget;
diff --git a/Engine/Plugins/Experimental/DataPrepEditor/Source/DataPrepEditor/Private/Widgets/DataprepWidgets.cpp b/Engine/Plugins/Experimental/DataPrepEditor/Source/DataPrepEditor/Private/Widgets/DataprepWidgets.cpp
--- a/Engine/Plugins/Experimental/DataPrepEditor/Source/DataPrepEditor/Private/Widgets/DataprepWidgets.cpp
+++ b/Engine/Plugins/Experimental/DataPrepEditor/Source/DataPrepEditor/Private/Widgets/DataprepWidgets.cpp
@@ -149,7 +149,7 @@ TSharedRef<SWidget> SDataprepConsumerWidget::BuildWidget()
 				.FillWidth(1.0f)
 				[
 					SNew(STextBlock)
-					.Text(LOCTEXT("DataprepSlateHelper_LevelNameLabel", "Level"))
+					.Text(LOCTEXT("DataprepSlateHelper_LevelNameLabel", "Sub-Level"))
 					.Font( IDetailLayoutBuilder::GetDetailFont() )
 				]
 			]
