Fixed hanging while taking snapshot on material

diff --git a/Engine/Plugins/Experimental/DataPrepEditor/Source/DataPrepEditor/Private/DataprepSnapshot.cpp b/Engine/Plugins/Experimental/DataPrepEditor/Source/DataPrepEditor/Private/DataprepSnapshot.cpp
--- a/Engine/Plugins/Experimental/DataPrepEditor/Source/DataPrepEditor/Private/DataprepSnapshot.cpp
+++ b/Engine/Plugins/Experimental/DataPrepEditor/Source/DataPrepEditor/Private/DataprepSnapshot.cpp
@@ -71,8 +71,8 @@ namespace DataprepSnapshotUtil
 		// Helper struct to identify dependency of a UObject on other UObject(s) except given one (its outer)
 		struct FObjectDependencyAnalyzer : public FArchiveUObject
 		{
-			FObjectDependencyAnalyzer(UObject* InObject)
-				: ObjectToSkip(InObject)
+			FObjectDependencyAnalyzer(const TSet<UObject*>& InValidObjects)
+				: ValidObjects( DependentObjects )
 			{ }
 
 			virtual FArchive& operator<<(UObject*& Obj) override
@@ -83,9 +83,9 @@ namespace DataprepSnapshotUtil
 					{
 						return FArchiveUObject::operator<<(Obj);
 					}
-					else if(Obj != ObjectToSkip && DependentObjects.Find(Obj) == nullptr)
+					else if( ValidObjects.Contains( Obj ) )
 					{
-						DependentObjects.Add(Obj);
+						DependentObjects.Add( Obj );
 					}
 				}
 
@@ -93,7 +93,7 @@ namespace DataprepSnapshotUtil
 			}
 
 			TSet<UObject*> DependentObjects;
-			UObject* ObjectToSkip;
+			const TSet<UObject*>& ValidObjects;
 		};
 
 		TArray<uint8> MemoryBuffer;
@@ -117,9 +117,10 @@ namespace DataprepSnapshotUtil
 			}
 
 			// Build graph of dependency: each entry contains the set of sub-objects to create before itself
+			TSet<UObject*> SubObjectsSet(SubObjectsArray);
 			for(UObject* SubObject : SubObjectsArray)
 			{
-				FObjectDependencyAnalyzer Analyzer(SubObject->GetOuter());
+				FObjectDependencyAnalyzer Analyzer( SubObjectsSet );
 				SubObject->Serialize( Analyzer );
 
 				SubObjectDependencyGraph[SubObject].Append(Analyzer.DependentObjects);
@@ -364,7 +365,7 @@ void FDataprepEditor::TakeSnapshot()
 	}
 
 	uint64 StartTime = FPlatformTime::Cycles64();
-	UE_LOG( LogDataprepEditor, Log, TEXT("Restoring snapshot...") );
+	UE_LOG( LogDataprepEditor, Log, TEXT("Taking snapshot...") );
 
 	// Clean up temporary folder with content of previous snapshot(s)
 	{
