The bounding volume fetcher is now not dependant on the static mesh being build.

diff --git a/Engine/Plugins/Experimental/DataPrepEditor/Source/DataPrepLibraries/Private/Fetchers/DataprepFloatFetcherLibrary.cpp b/Engine/Plugins/Experimental/DataPrepEditor/Source/DataPrepLibraries/Private/Fetchers/DataprepFloatFetcherLibrary.cpp
--- a/Engine/Plugins/Experimental/DataPrepEditor/Source/DataPrepLibraries/Private/Fetchers/DataprepFloatFetcherLibrary.cpp
+++ b/Engine/Plugins/Experimental/DataPrepEditor/Source/DataPrepLibraries/Private/Fetchers/DataprepFloatFetcherLibrary.cpp
@@ -5,6 +5,9 @@
 #include "Components/StaticMeshComponent.h"
 #include "Engine/StaticMesh.h"
 #include "GameFramework/Actor.h"
+#include "MeshAttributeArray.h"
+#include "MeshAttributes.h"
+#include "MeshDescription.h"
 #include "Misc/Optional.h"
 #include "UObject/Object.h"
 
@@ -14,45 +17,58 @@ float UDataprepFloatBoundingVolumeFetcher::Fetch_Implementation(const UObject* O
 {
 	if ( Object && !Object->IsPendingKill() )
 	{
-		// build the static mesh if there isn't a render data (this a temporary work around)
-		auto BuildStaticMeshIfNoRenderData = [](UStaticMesh* StaticMesh)
+		auto GetStaticMeshBoundingBox = [](const UStaticMesh* StaticMesh) -> FBox
 			{
 				if ( !StaticMesh->RenderData || !StaticMesh->RenderData->IsInitialized() )
 				{
-					const_cast< UStaticMesh* >( StaticMesh )->Build( true );
+					FBox Box(ForceInit);
+					if ( const FMeshDescription* MeshDescription = StaticMesh->GetMeshDescription( 0 ) )
+					{
+						TVertexAttributesConstRef<FVector> VertexPositions = MeshDescription->VertexAttributes().GetAttributesRef<FVector>( MeshAttribute::Vertex::Position );
+						for ( const FVertexID VertexID : MeshDescription->Vertices().GetElementIDs() )
+						{
+							if ( !MeshDescription->IsVertexOrphaned( VertexID ) )
+							{
+								Box += VertexPositions[VertexID];
+							}
+						}
+					}
+					return Box;
 				}
+
+				return StaticMesh->GetBoundingBox();
 			};
 
 		TOptional<float> Volume;
 		if ( const AActor* Actor = Cast<const AActor>( Object ) )
 		{
-			FBox Box(ForceInit);
+			FBox ActorBox(ForceInit);
 
 			for ( const UActorComponent* ActorComponent : Actor->GetComponents() )
 			{
 				if ( const UPrimitiveComponent* PrimComp = Cast<const UPrimitiveComponent>( ActorComponent ) )
 				{
+					FBox ComponentBox(ForceInit);
 					if ( const UStaticMeshComponent* StaticMeshComponent = Cast<const UStaticMeshComponent>( ActorComponent ) )
 					{
-						BuildStaticMeshIfNoRenderData( StaticMeshComponent->GetStaticMesh() );
+						ComponentBox = GetStaticMeshBoundingBox( StaticMeshComponent->GetStaticMesh() );
+						ComponentBox.TransformBy( PrimComp->GetComponentToWorld() );
 					}
 					if ( PrimComp->IsRegistered() &&  PrimComp->IsCollisionEnabled() )
 					{
-						
-						Box += PrimComp->Bounds.GetBox();
+						ActorBox += ComponentBox.IsValid? ComponentBox : PrimComp->Bounds.GetBox();
 					}
 				}
 			}
 
-			if ( Box.IsValid )
+			if ( ActorBox.IsValid )
 			{
-				Volume = Box.GetVolume();
+				Volume = ActorBox.GetVolume();
 			}
 		}
 		else if ( const UStaticMesh* StaticMesh = Cast<const UStaticMesh>( Object ) )
 		{
-			BuildStaticMeshIfNoRenderData( const_cast<UStaticMesh*>( StaticMesh ) );
-			Volume = StaticMesh->GetBoundingBox().GetVolume();
+			Volume = GetStaticMeshBoundingBox( StaticMesh ).GetVolume();
 		}
 
 		if ( Volume.IsSet() )
@@ -68,6 +84,5 @@ float UDataprepFloatBoundingVolumeFetcher::Fetch_Implementation(const UObject* O
 
 bool UDataprepFloatBoundingVolumeFetcher::IsThreadSafe() const
 {
-	// For 4.23 it will not be multi thread safe since mesh aren't always build
-	return false;
+	return true;
 }
