Fixed crash when calling SetMaterial Dataprep operation

diff --git a/Engine/Plugins/Experimental/DataPrepEditor/Source/DataPrepLibraries/Private/DataPrepOperationsLibrary.cpp b/Engine/Plugins/Experimental/DataPrepEditor/Source/DataPrepLibraries/Private/DataPrepOperationsLibrary.cpp
--- a/Engine/Plugins/Experimental/DataPrepEditor/Source/DataPrepLibraries/Private/DataPrepOperationsLibrary.cpp
+++ b/Engine/Plugins/Experimental/DataPrepEditor/Source/DataPrepLibraries/Private/DataPrepOperationsLibrary.cpp
@@ -26,6 +26,7 @@
 #include "MeshTypes.h"
 #include "Misc/FileHelper.h"
 #include "ObjectTools.h"
+#include "TessellationRendering.h"
 #include "UObject/SoftObjectPath.h"
 
 DEFINE_LOG_CATEGORY(LogDataprep);
@@ -188,6 +189,32 @@ namespace DataprepOperationsLibraryUtil
 	{
 		return Actor ? 1 + GetActorDepth(Actor->GetAttachParentActor()) : 0;
 	}
+
+	/** Customized version of UStaticMesh::SetMaterial avoiding the triggering of UStaticMesh::Build and its side-effects */
+	void SetMaterial( UStaticMesh* StaticMesh, int32 MaterialIndex, UMaterialInterface* NewMaterial )
+	{
+		if( StaticMesh->StaticMaterials.IsValidIndex( MaterialIndex ) )
+		{
+			FStaticMaterial& StaticMaterial = StaticMesh->StaticMaterials[ MaterialIndex ];
+			StaticMaterial.MaterialInterface = NewMaterial;
+			if( NewMaterial != nullptr )
+			{
+				if ( StaticMaterial.MaterialSlotName == NAME_None )
+				{
+					StaticMaterial.MaterialSlotName = NewMaterial->GetFName();
+				}
+
+				// Make sure adjacency information fit new material change
+				if( RequiresAdjacencyInformation( NewMaterial, nullptr, GWorld->FeatureLevel ) )
+				{
+					for( FStaticMeshSourceModel& SourceModel : StaticMesh->SourceModels )
+					{
+						SourceModel.BuildSettings.bBuildAdjacencyBuffer = true;
+					}
+				}
+			}
+		}
+	}
 } // ns DataprepOperationsLibraryUtil
 
 void UDataprepOperationsLibrary::SetLods(const TArray<UObject*>& SelectedObjects, const FEditorScriptingMeshReductionOptions& ReductionOptions)
@@ -343,7 +370,7 @@ void UDataprepOperationsLibrary::SubstituteMaterial(const TArray<UObject*>& Sele
 				{
 					if (StaticMesh->GetMaterial(Index) == MaterialToReplace)
 					{
-						StaticMesh->SetMaterial(Index, MaterialSubstitute);
+						DataprepOperationsLibraryUtil::SetMaterial( StaticMesh, Index, MaterialSubstitute );
 					}
 				}
 			}
@@ -389,10 +416,9 @@ void UDataprepOperationsLibrary::SetMaterial( const TArray< UObject* >& Selected
 		{
 			DataprepOperationsLibraryUtil::FScopedStaticMeshEdit StaticMeshEdit( StaticMesh );
 
-			TArray<FStaticMaterial>& StaticMaterials = StaticMesh->StaticMaterials;
 			for (int32 Index = 0; Index < StaticMesh->StaticMaterials.Num(); ++Index)
 			{
-				StaticMesh->SetMaterial(Index, MaterialSubstitute);
+				DataprepOperationsLibraryUtil::SetMaterial( StaticMesh, Index, MaterialSubstitute );
 			}
 		}
 	}
