Fixed crash when undoing dataprep graph node deletion

diff --git a/Engine/Plugins/Experimental/DataPrepEditor/Source/DataPrepEditor/Private/DataPrepEditor.cpp b/Engine/Plugins/Experimental/DataPrepEditor/Source/DataPrepEditor/Private/DataPrepEditor.cpp
--- a/Engine/Plugins/Experimental/DataPrepEditor/Source/DataPrepEditor/Private/DataPrepEditor.cpp
+++ b/Engine/Plugins/Experimental/DataPrepEditor/Source/DataPrepEditor/Private/DataPrepEditor.cpp
@@ -300,6 +300,13 @@ void FDataprepEditor::InitDataprepEditor(const EToolkitMode::Type Mode, const TS
 	DataprepRecipeBPPtr = DataprepAssetPtr->DataprepRecipeBP;
 	check( DataprepRecipeBPPtr.IsValid() );
 
+	// Necessary step to regenerate blueprint generated class
+	// Note that this compilation will always succeed as Dataprep node does not have real body
+	// #ueent_todo: Is there a better solution
+	{
+		FKismetEditorUtilities::CompileBlueprint( DataprepRecipeBPPtr.Get(), EBlueprintCompileOptions::None, nullptr );
+	}
+
 	UEdGraph* PipelineGraph = FBlueprintEditorUtils::FindEventGraph(DataprepRecipeBPPtr.Get());
 	check( PipelineGraph );
 
