Dataprep Misc UI Clean Up

diff --git a/Engine/Plugins/Experimental/DataPrepEditor/Source/DataPrepEditor/Private/DataPrepEditor.cpp b/Engine/Plugins/Experimental/DataPrepEditor/Source/DataPrepEditor/Private/DataPrepEditor.cpp
--- a/Engine/Plugins/Experimental/DataPrepEditor/Source/DataPrepEditor/Private/DataPrepEditor.cpp
+++ b/Engine/Plugins/Experimental/DataPrepEditor/Source/DataPrepEditor/Private/DataPrepEditor.cpp
@@ -373,12 +373,14 @@ void FDataprepEditor::InitDataprepEditor(const EToolkitMode::Type Mode, const TS
 							FTabManager::NewStack()
 							->SetSizeCoefficient(0.5f)
 							->AddTab(ScenePreviewTabId, ETabState::OpenedTab)
+							->SetHideTabWell( true )
 						)
 						->Split
 						(
 							FTabManager::NewStack()
 							->SetSizeCoefficient(0.5f)
 							->AddTab(AssetPreviewTabId, ETabState::OpenedTab)
+							->SetHideTabWell( true )
 						)
 					)
 					->Split
@@ -389,6 +391,7 @@ void FDataprepEditor::InitDataprepEditor(const EToolkitMode::Type Mode, const TS
 							FTabManager::NewStack()
 							->SetSizeCoefficient(0.15f)
 							->AddTab(PaletteTabId, ETabState::OpenedTab)
+							->SetHideTabWell( true )
 						)
 						// Temp code for the nodes development
 						->Split
@@ -396,6 +399,7 @@ void FDataprepEditor::InitDataprepEditor(const EToolkitMode::Type Mode, const TS
 							FTabManager::NewStack()
 							->SetSizeCoefficient(0.85f)
 							->AddTab(PipelineGraphTabId, ETabState::OpenedTab)
+							->SetHideTabWell( true )
 						)
 						// end of temp code for nodes development
 					)
@@ -409,6 +413,7 @@ void FDataprepEditor::InitDataprepEditor(const EToolkitMode::Type Mode, const TS
 						FTabManager::NewStack()
 						->SetSizeCoefficient(0.25f)
 						->AddTab(DataprepAssetTabId, ETabState::OpenedTab)
+						->SetHideTabWell(true)
 					)
 					->Split
 					(
diff --git a/Engine/Plugins/Experimental/DataPrepEditor/Source/DataPrepEditor/Private/DataprepEditorStyle.cpp b/Engine/Plugins/Experimental/DataPrepEditor/Source/DataPrepEditor/Private/DataprepEditorStyle.cpp
--- a/Engine/Plugins/Experimental/DataPrepEditor/Source/DataPrepEditor/Private/DataprepEditorStyle.cpp
+++ b/Engine/Plugins/Experimental/DataPrepEditor/Source/DataPrepEditor/Private/DataprepEditorStyle.cpp
@@ -67,7 +67,7 @@ void FDataprepEditorStyle::Initialize()
 		TilteTextBlockStyle.SetFont(FCoreStyle::GetDefaultFontStyle("Bold", 14));
 		StyleSet->Set("DataprepActionBlock.TitleTextBlockStyle", TilteTextBlockStyle);
 		StyleSet->Set("DataprepActionSteps.BackgroundColor", FLinearColor( 0.01033, 0.01033, 0.01033));
-		StyleSet->Set("DataprepActionStep.Padding", 25.f);
+		StyleSet->Set("DataprepActionStep.Padding", 10.f);
 	}
 
 	FSlateStyleRegistry::RegisterSlateStyle(*StyleSet.Get());
diff --git a/Engine/Plugins/Experimental/DataPrepEditor/Source/DataPrepEditor/Private/Widgets/SDataprepPalette.cpp b/Engine/Plugins/Experimental/DataPrepEditor/Source/DataPrepEditor/Private/Widgets/SDataprepPalette.cpp
--- a/Engine/Plugins/Experimental/DataPrepEditor/Source/DataPrepEditor/Private/Widgets/SDataprepPalette.cpp
+++ b/Engine/Plugins/Experimental/DataPrepEditor/Source/DataPrepEditor/Private/Widgets/SDataprepPalette.cpp
@@ -37,13 +37,6 @@ void SDataprepPalette::Construct(const FArguments& InArgs)
 	AllCategory = LOCTEXT("All Category", "All");
 	SelectorsCategory = FDataprepFilterMenuActionCollector::FilterCategory;
 	OperationsCategory = FDataprepOperationMenuActionCollector::OperationCategory;
-	AllCategoryPtr = MakeShared< FString >( AllCategory.ToString() );
-	SelectorsCategoryPtr = MakeShared< FString >( SelectorsCategory.ToString() );
-	OperationsCategoryPtr = MakeShared< FString >( OperationsCategory.ToString() );
-	SelectedCategoryPtr = AllCategoryPtr;
-	CategoryNames.Add( AllCategoryPtr );
-	CategoryNames.Add( SelectorsCategoryPtr );
-	CategoryNames.Add( OperationsCategoryPtr );
 
 	this->ChildSlot
 		[
@@ -52,30 +45,6 @@ void SDataprepPalette::Construct(const FArguments& InArgs)
 			.BorderImage( FEditorStyle::GetBrush("ToolPanel.GroupBorder") )
 			[
 				SNew(SVerticalBox)
-				
-				// UI for filtering by selector , operation or all
-				+ SVerticalBox::Slot()
-				.AutoHeight()
-				[
-					SNew(SHorizontalBox)
-					+ SHorizontalBox::Slot()
-					.VAlign(VAlign_Center)
-					.AutoWidth()
-					[
-						SNew(STextBlock)
-						.Text( LOCTEXT("Category", "Category: ") )
-					]
-					// Combo button to select a class
-					+ SHorizontalBox::Slot()
-					.VAlign(VAlign_Center)
-					[
-							SNew(STextComboBox)
-							.OptionsSource( &CategoryNames )
-							.OnSelectionChanged( this, &SDataprepPalette::CategorySelectionChanged )
-							.InitiallySelectedItem( CategoryNames[0] )
-					]
-				]
-
 				// Content list
 				+ SVerticalBox::Slot()
 				[
@@ -115,29 +84,10 @@ void SDataprepPalette::Construct(const FArguments& InArgs)
 
 void SDataprepPalette::CollectAllActions(FGraphActionListBuilderBase& OutAllActions)
 {
-	TUniquePtr<IDataprepMenuActionCollector> ActionCollector;
-
-	if ( SelectedCategoryPtr == AllCategoryPtr )
-	{
-		ActionCollector = MakeUnique< FDataprepAllMenuActionCollector >();
-		
-	}
-	else if ( SelectedCategoryPtr == SelectorsCategoryPtr )
+	FDataprepAllMenuActionCollector ActionCollector;
+	for ( TSharedPtr< FDataprepSchemaAction > Action : ActionCollector.CollectActions() )
 	{
-		ActionCollector = MakeUnique< FDataprepFilterMenuActionCollector >();
-		
-	}
-	else if	( SelectedCategoryPtr == OperationsCategoryPtr )
-	{
-		ActionCollector = MakeUnique< FDataprepOperationMenuActionCollector >();
-	}
-
-	if ( ActionCollector )
-	{
-		for ( TSharedPtr< FDataprepSchemaAction > Action : ActionCollector->CollectActions() )
-		{
-			OutAllActions.AddAction(StaticCastSharedPtr< FEdGraphSchemaAction >(Action));
-		}
+		OutAllActions.AddAction(StaticCastSharedPtr< FEdGraphSchemaAction >(Action));
 	}
 }
 
@@ -160,15 +110,6 @@ TSharedRef<SExpanderArrow> SDataprepPalette::OnCreateCustomRowExpander(const FCu
 	return SNew(SExpanderArrow, InCustomExpanderData.TableRow);
 }
 
-void SDataprepPalette::CategorySelectionChanged(TSharedPtr<FString> NewSelection, ESelectInfo::Type SelectInfo)
-{
-	if ( NewSelection )
-	{
-		SelectedCategoryPtr = NewSelection;
-		RefreshActionsList( true );
-	}
-}
-
 void SDataprepPalette::AddAssetFromAssetRegistry(const FAssetData& InAddedAssetData)
 {
 	RefreshAssetInRegistry( InAddedAssetData );
diff --git a/Engine/Plugins/Experimental/DataPrepEditor/Source/DataPrepEditor/Private/Widgets/SDataprepPalette.h b/Engine/Plugins/Experimental/DataPrepEditor/Source/DataPrepEditor/Private/Widgets/SDataprepPalette.h
--- a/Engine/Plugins/Experimental/DataPrepEditor/Source/DataPrepEditor/Private/Widgets/SDataprepPalette.h
+++ b/Engine/Plugins/Experimental/DataPrepEditor/Source/DataPrepEditor/Private/Widgets/SDataprepPalette.h
@@ -28,9 +28,6 @@ private:
 	/** CallBack from the GraphActionMenu for when the custom row expender is created */
 	TSharedRef<class SExpanderArrow> OnCreateCustomRowExpander(const FCustomExpanderData& InCustomExpanderData) const;
 
-	/** Callback for when the category selection is changed */
-	void CategorySelectionChanged(TSharedPtr<FString> NewSelection, ESelectInfo::Type SelectInfo);
-
 	/** Callback from the Asset Registry when a new asset is added. */
 	void AddAssetFromAssetRegistry(const FAssetData& InAddedAssetData);
 
@@ -42,16 +39,9 @@ private:
 
 	void RefreshAssetInRegistry(const FAssetData& InAssetData);
 
-	// The currently selected category 
-	TSharedPtr< FString > SelectedCategoryPtr;
-
 	// The options for the category selection
 	FText AllCategory;
 	FText SelectorsCategory;
 	FText OperationsCategory;
-	TSharedPtr< FString > AllCategoryPtr;
-	TSharedPtr< FString > SelectorsCategoryPtr;
-	TSharedPtr< FString > OperationsCategoryPtr;
-	TArray< TSharedPtr<FString> > CategoryNames;
 };
 
