Fixed no warning when committing if last execute is out of sync with current pipeline graph
On behalf of julien.stjean

diff --git a/Engine/Plugins/Experimental/DataPrepEditor/Source/DataPrepCore/Public/DataPrepAsset.h b/Engine/Plugins/Experimental/DataPrepEditor/Source/DataPrepCore/Public/DataPrepAsset.h
--- a/Engine/Plugins/Experimental/DataPrepEditor/Source/DataPrepCore/Public/DataPrepAsset.h
+++ b/Engine/Plugins/Experimental/DataPrepEditor/Source/DataPrepCore/Public/DataPrepAsset.h
@@ -136,6 +136,26 @@ public:
 		return Producers.IsValidIndex(Index) ? Producers[Index].SupersededBy != INDEX_NONE && Producers.IsValidIndex(Producers[Index].SupersededBy) && Producers[Producers[Index].SupersededBy].bIsEnabled : false;
 	}
 
+
+	/**
+	 * Allow an observer to be notified of an change in the pipeline
+	 * return The event that will be broadcasted when a object has receive a modification that might change the result of the pipeline
+	 */
+	DECLARE_EVENT_OneParam(UDataprepAsset, FOnDataprepPipelineChange, UObject* /*The object that was modified*/)
+	FOnDataprepPipelineChange& GetOnPipelineChange() { return OnPipelineChange; }
+
+	// struct to restrict the access scope
+	struct FDataprepPipelineChangeNotifier
+	{
+	private:
+		friend class FDataprepEditorUtils;
+
+		static void NotifyDataprepOfPipelineChange(UDataprepAsset& DataprepAsset, UObject* ModifiedObject)
+		{
+			DataprepAsset.OnPipelineChange.Broadcast( ModifiedObject );
+		}
+	};
+
 	/**
 	 * Allow an observer to be notified when when the consumer or one of the producer has changed
 	 * @return The delegate that will be broadcasted when when the consumer or one of the producer has changed
@@ -188,6 +208,10 @@ private:
 	  */
 	void ValidateProducerChanges( int32 Index, bool &bChangeAll );
 
+
 	/** Delegate broadcasted when the consumer or one of the producers has changed */
 	FOnDataprepAssetChanged OnChanged;
+
+	/** Event broadcasted when object in the pipeline was modified (Only broadcasted on changes that can affect the result of execution) */
+	FOnDataprepPipelineChange OnPipelineChange;
 };
\ No newline at end of file
diff --git a/Engine/Plugins/Experimental/DataPrepEditor/Source/DataPrepEditor/Private/BlueprintNodes/K2Node_DataprepAction.cpp b/Engine/Plugins/Experimental/DataPrepEditor/Source/DataPrepEditor/Private/BlueprintNodes/K2Node_DataprepAction.cpp
--- a/Engine/Plugins/Experimental/DataPrepEditor/Source/DataPrepEditor/Private/BlueprintNodes/K2Node_DataprepAction.cpp
+++ b/Engine/Plugins/Experimental/DataPrepEditor/Source/DataPrepEditor/Private/BlueprintNodes/K2Node_DataprepAction.cpp
@@ -5,6 +5,7 @@
 // Dataprep includes
 #include "DataPrepAsset.h"
 #include "DataprepActionAsset.h"
+#include "DataprepEditorUtils.h"
 #include "Widgets/Action/SGraphNodeK2DataprepAction.h"
 
 // Engine includes
@@ -102,6 +103,11 @@ void UK2Node_DataprepAction::DestroyNode()
 	Super::DestroyNode();
 }
 
+void UK2Node_DataprepAction::NodeConnectionListChanged()
+{
+	FDataprepEditorUtils::NotifySystemOfChangeInPipeline( this );
+}
+
 TSharedPtr<SGraphNode> UK2Node_DataprepAction::CreateVisualWidget()
 {
 	return SNew( SGraphNodeK2DataprepAction, this );
diff --git a/Engine/Plugins/Experimental/DataPrepEditor/Source/DataPrepEditor/Private/BlueprintNodes/K2Node_DataprepAction.h b/Engine/Plugins/Experimental/DataPrepEditor/Source/DataPrepEditor/Private/BlueprintNodes/K2Node_DataprepAction.h
--- a/Engine/Plugins/Experimental/DataPrepEditor/Source/DataPrepEditor/Private/BlueprintNodes/K2Node_DataprepAction.h
+++ b/Engine/Plugins/Experimental/DataPrepEditor/Source/DataPrepEditor/Private/BlueprintNodes/K2Node_DataprepAction.h
@@ -30,6 +30,7 @@ public:
 	virtual FText GetNodeTitle(ENodeTitleType::Type TitleType) const override;
 	virtual void OnRenameNode(const FString& NewName) override;
 	virtual void DestroyNode() override;
+	virtual void NodeConnectionListChanged() override;
 	virtual TSharedPtr<SGraphNode> CreateVisualWidget() override;
 	TSharedPtr<class INameValidatorInterface> MakeNameValidator() const override;
 	// End EdGraphNode interface
diff --git a/Engine/Plugins/Experimental/DataPrepEditor/Source/DataPrepEditor/Private/DataPrepEditor.cpp b/Engine/Plugins/Experimental/DataPrepEditor/Source/DataPrepEditor/Private/DataPrepEditor.cpp
--- a/Engine/Plugins/Experimental/DataPrepEditor/Source/DataPrepEditor/Private/DataPrepEditor.cpp
+++ b/Engine/Plugins/Experimental/DataPrepEditor/Source/DataPrepEditor/Private/DataPrepEditor.cpp
@@ -24,6 +24,7 @@
 #include "BlueprintNodeSpawner.h"
 #include "DesktopPlatformModule.h"
 #include "Dialogs/DlgPickPath.h"
+#include "Dialogs/Dialogs.h"
 #include "Editor.h"
 #include "EditorDirectories.h"
 #include "EditorStyleSet.h"
@@ -46,6 +47,7 @@
 #include "ScopedTransaction.h"
 #include "Templates/UnrealTemplate.h"
 #include "Toolkits/IToolkit.h"
+#include "UObject/Object.h"
 #include "UnrealEdGlobals.h"
 #include "Widgets/Docking/SDockTab.h"
 #include "Widgets/Input/STextComboBox.h"
@@ -155,7 +157,7 @@ private:
 FDataprepEditor::FDataprepEditor()
 	: bWorldBuilt(false)
 	, bIsFirstRun(false)
-	, bPipelineExecuted(false)
+	, bPipelineChanged(false)
 	, bIsActionMenuContextSensitive(true)
 	, bSaveIntermediateBuildProducts(false)
 	, PreviewWorld(nullptr)
@@ -192,6 +194,7 @@ FDataprepEditor::~FDataprepEditor()
 	if( DataprepAssetPtr.IsValid() )
 	{
 		DataprepAssetPtr->GetOnChanged().RemoveAll( this );
+		DataprepAssetPtr->GetOnPipelineChange().RemoveAll( this );
 	}
 
 	if ( PreviewWorld )
@@ -288,6 +291,8 @@ void FDataprepEditor::InitDataprepEditor(const EToolkitMode::Type Mode, const TS
 	check( DataprepAssetPtr.IsValid() );
 
 	DataprepAssetPtr->GetOnChanged().AddRaw( this, &FDataprepEditor::OnDataprepAssetChanged );
+	DataprepAssetPtr->GetOnPipelineChange().AddRaw( this, &FDataprepEditor::OnDataprepPipelineChange );
+
 
 	// Assign unique session identifier
 	SessionID = FGuid::NewGuid().ToString();
@@ -522,7 +527,6 @@ void FDataprepEditor::OnBuildWorld()
 	UpdatePreviewPanels();
 	bWorldBuilt = true;
 	bIsFirstRun = true;
-	bPipelineExecuted = false;
 
 	// Log time spent to import incoming file in minutes and seconds
 	double ElapsedSeconds = FPlatformTime::ToSeconds64(FPlatformTime::Cycles64() - StartTime);
@@ -546,7 +550,7 @@ void FDataprepEditor::OnDataprepAssetChanged( FDataprepAssetChangeType ChangeTyp
 
 		case FDataprepAssetChangeType::BlueprintModified:
 			{
-				// #ueent_todo: Anything to do there ?
+				OnDataprepPipelineChange( nullptr );
 			}
 			break;
 
@@ -565,6 +569,11 @@ void FDataprepEditor::OnDataprepAssetChanged( FDataprepAssetChangeType ChangeTyp
 	}
 }
 
+void FDataprepEditor::OnDataprepPipelineChange(UObject* ChangedObject)
+{
+	bPipelineChanged = true;
+}
+
 void FDataprepEditor::ResetBuildWorld()
 {
 	bWorldBuilt = false;
@@ -722,16 +731,39 @@ void FDataprepEditor::OnExecutePipeline()
 		}
 	}
 
-	// Indicate pipeline has been executed
+	// Indicate pipeline has been executed at least once
 	bIsFirstRun = false;
-	bPipelineExecuted = true;
+	// Reset tracking of pipeline changes between execution
+	bPipelineChanged = false;
 }
 
 void FDataprepEditor::OnCommitWorld()
 {
-	if(!bPipelineExecuted)
+	// Pipeline has not been executed, validate with user this is intentional
+	if( bIsFirstRun )
 	{
-		// #ueent_todo Prompt user if pipeline has not been run on imported assets
+		UEdGraphPin* StartNodePin = StartNode->FindPin( UEdGraphSchema_K2::PN_Then, EGPD_Output );
+		if(StartNodePin && StartNodePin->LinkedTo.Num() > 0)
+		{
+			const FText Title( LOCTEXT( "DataprepEditor_ProceedWithCommit", "Proceed with commit" ) );
+			const FText Message( LOCTEXT( "DataprepEditor_ConfirmCommitPipelineNotExecuted", "The action pipeline has not been executed.\nDo you want to proceeed with the commit anyway?" ) );
+
+			if( OpenMsgDlgInt( EAppMsgType::YesNo, Message, Title ) == EAppReturnType::No )
+			{
+				return;
+			}
+		}
+	}
+	// Pipeline has changed without being executed, validate with user this is intentional
+	else if( !bIsFirstRun && bPipelineChanged )
+	{
+		const FText Title( LOCTEXT( "DataprepEditor_ProceedWithCommit", "Proceed with commit" ) );
+		const FText Message( LOCTEXT( "DataprepEditor_ConfirmCommitPipelineChanged", "The action pipeline has changed since last execution.\nDo you want to proceeed with the commit anyway?" ) );
+
+		if( OpenMsgDlgInt( EAppMsgType::YesNo, Message, Title ) == EAppReturnType::No )
+		{
+			return;
+		}
 	}
 
 	// Finalize assets
diff --git a/Engine/Plugins/Experimental/DataPrepEditor/Source/DataPrepEditor/Private/DataPrepEditor.h b/Engine/Plugins/Experimental/DataPrepEditor/Source/DataPrepEditor/Private/DataPrepEditor.h
--- a/Engine/Plugins/Experimental/DataPrepEditor/Source/DataPrepEditor/Private/DataPrepEditor.h
+++ b/Engine/Plugins/Experimental/DataPrepEditor/Source/DataPrepEditor/Private/DataPrepEditor.h
@@ -146,12 +146,15 @@ private:
 	void RestoreFromSnapshot();
 
 	/** Handles changes in the Dataprep asset */
-	void OnDataprepAssetChanged(FDataprepAssetChangeType ChangeType, int32 Index );
+	void OnDataprepAssetChanged(FDataprepAssetChangeType ChangeType, int32 Index);
+
+	/** Handles change to the dataprep pipeline */
+	void OnDataprepPipelineChange(UObject* ChangedObject);
 
 private:
 	bool bWorldBuilt;
 	bool bIsFirstRun;
-	bool bPipelineExecuted;
+	bool bPipelineChanged;
 	TWeakObjectPtr<UDataprepAsset> DataprepAssetPtr;
 
 	FOnDataprepAssetProducerChanged DataprepAssetProducerChangedDelegate;
diff --git a/Engine/Plugins/Experimental/DataPrepEditor/Source/DataPrepEditor/Private/DataprepEditorUtils.cpp b/Engine/Plugins/Experimental/DataPrepEditor/Source/DataPrepEditor/Private/DataprepEditorUtils.cpp
new file mode 100644
--- /dev/null
+++ b/Engine/Plugins/Experimental/DataPrepEditor/Source/DataPrepEditor/Private/DataprepEditorUtils.cpp
@@ -0,0 +1,38 @@
+// Copyright 1998-2019 Epic Games, Inc. All Rights Reserved.
+
+#include "DataprepEditorUtils.h"
+
+#include "DataPrepAsset.h"
+#include "Engine/Blueprint.h"
+#include "Kismet2/BlueprintEditorUtils.h"
+
+void FDataprepEditorUtils::NotifySystemOfChangeInPipeline(UObject* SourceObject)
+{
+	UBlueprint* Blueprint = nullptr;
+	UDataprepAsset* DataprepAsset = nullptr;
+	UObject* Object = SourceObject;
+	while ( Object )
+	{
+		UClass* Class = Object->GetClass();
+		if ( Class->IsChildOf<UBlueprint>() )
+		{
+			Blueprint = static_cast<UBlueprint*>( Object );
+		}
+		else if ( Class == UDataprepAsset::StaticClass() )
+		{
+			DataprepAsset = static_cast<UDataprepAsset*>( Object );
+			break;
+		}
+		Object = Object->GetOuter();
+	}
+
+	if ( DataprepAsset )
+	{
+		UDataprepAsset::FDataprepPipelineChangeNotifier::NotifyDataprepOfPipelineChange( *DataprepAsset, SourceObject );
+	}
+	else if ( Blueprint )
+	{
+		FBlueprintEditorUtils::MarkBlueprintAsModified( Blueprint );
+	}
+}
+
diff --git a/Engine/Plugins/Experimental/DataPrepEditor/Source/DataPrepEditor/Private/DataprepEditorUtils.h b/Engine/Plugins/Experimental/DataPrepEditor/Source/DataPrepEditor/Private/DataprepEditorUtils.h
new file mode 100644
--- /dev/null
+++ b/Engine/Plugins/Experimental/DataPrepEditor/Source/DataPrepEditor/Private/DataprepEditorUtils.h
@@ -0,0 +1,11 @@
+// Copyright 1998-2019 Epic Games, Inc. All Rights Reserved.
+
+#pragma once
+
+#include "CoreMinimal.h"
+
+class FDataprepEditorUtils
+{
+public:
+	static void NotifySystemOfChangeInPipeline(UObject* SourceObject);
+};
diff --git a/Engine/Plugins/Experimental/DataPrepEditor/Source/DataPrepEditor/Private/SchemaActions/DataprepDragDropOp.cpp b/Engine/Plugins/Experimental/DataPrepEditor/Source/DataPrepEditor/Private/SchemaActions/DataprepDragDropOp.cpp
--- a/Engine/Plugins/Experimental/DataPrepEditor/Source/DataPrepEditor/Private/SchemaActions/DataprepDragDropOp.cpp
+++ b/Engine/Plugins/Experimental/DataPrepEditor/Source/DataPrepEditor/Private/SchemaActions/DataprepDragDropOp.cpp
@@ -9,6 +9,7 @@
 #include "EdGraph/EdGraph.h"
 #include "EdGraphSchema_K2.h"
 #include "EditorStyleSet.h"
+#include "Kismet2/BlueprintEditorUtils.h"
 #include "ScopedTransaction.h"
 #include "Templates/UnrealTypeTraits.h"
 #include "Widgets/Text/STextBlock.h"
@@ -161,6 +162,11 @@ void FDataprepDragDropOp::DoDropOnPanel(const TSharedRef<SWidget>& Panel, FVecto
 		{
 			Transaction.Cancel();
 		}
+
+		if ( UBlueprint* Blueprint = FBlueprintEditorUtils::FindBlueprintForGraphChecked( EdGraph ) )
+		{
+			FBlueprintEditorUtils::MarkBlueprintAsModified( Blueprint );
+		}
 	}
 }
 #undef LOCTEXT_NAMESPACE
diff --git a/Engine/Plugins/Experimental/DataPrepEditor/Source/DataPrepEditor/Private/SchemaActions/DataprepFetcherMenuActionCollector.cpp b/Engine/Plugins/Experimental/DataPrepEditor/Source/DataPrepEditor/Private/SchemaActions/DataprepFetcherMenuActionCollector.cpp
--- a/Engine/Plugins/Experimental/DataPrepEditor/Source/DataPrepEditor/Private/SchemaActions/DataprepFetcherMenuActionCollector.cpp
+++ b/Engine/Plugins/Experimental/DataPrepEditor/Source/DataPrepEditor/Private/SchemaActions/DataprepFetcherMenuActionCollector.cpp
@@ -2,6 +2,7 @@
 
 #include "SchemaActions/DataprepFetcherMenuActionCollector.h"
 
+#include "DataprepEditorUtils.h"
 #include "DataprepMenuActionCollectorUtils.h"
 #include "SelectionSystem/DataprepFetcher.h"
 
@@ -47,6 +48,7 @@ TSharedPtr<FDataprepSchemaAction> FDataprepFetcherMenuActionCollector::CreateMen
 		if ( Filter && ( !Filter->GetFetcher() || Filter->GetFetcher()->GetClass() != Class ) )
 		{
 			Filter->SetFetcher( TSubclassOf< UDataprepFetcher >( Class ) );
+			FDataprepEditorUtils::NotifySystemOfChangeInPipeline( Filter );
 		}
 	});
 
diff --git a/Engine/Plugins/Experimental/DataPrepEditor/Source/DataPrepEditor/Private/Widgets/Action/SDataprepActionSteps.cpp b/Engine/Plugins/Experimental/DataPrepEditor/Source/DataPrepEditor/Private/Widgets/Action/SDataprepActionSteps.cpp
--- a/Engine/Plugins/Experimental/DataPrepEditor/Source/DataPrepEditor/Private/Widgets/Action/SDataprepActionSteps.cpp
+++ b/Engine/Plugins/Experimental/DataPrepEditor/Source/DataPrepEditor/Private/Widgets/Action/SDataprepActionSteps.cpp
@@ -5,8 +5,9 @@
 #include "DataPrepOperation.h"
 #include "DataprepActionAsset.h"
 #include "DataprepEditorStyle.h"
-#include "SchemaActions/DataprepDragDropOp.h"
+#include "DataprepEditorUtils.h"
 #include "SchemaActions/DataprepAllMenuActionCollector.h"
+#include "SchemaActions/DataprepDragDropOp.h"
 #include "SchemaActions/DataprepSchemaAction.h"
 #include "SelectionSystem/DataprepFilter.h"
 #include "Widgets/Action/SDataprepFilter.h"
@@ -196,10 +197,16 @@ void SDataprepActionSteps::Construct(const FArguments& InArgs, UDataprepActionAs
 
 	if ( InDataprepAction )
 	{
-		InDataprepAction->GetOnStepsOrderChanged().AddSP( this, &SDataprepActionSteps::Refresh );
+		InDataprepAction->GetOnStepsOrderChanged().AddSP( this, &SDataprepActionSteps::OnStepsOrderChanged );
 	}
 }
 
+void SDataprepActionSteps::OnStepsOrderChanged()
+{
+	Refresh();
+	FDataprepEditorUtils::NotifySystemOfChangeInPipeline( DataprepActionPtr.Get() );
+}
+
 void SDataprepActionSteps::Refresh()
 {
 	UDataprepActionAsset* DataprepAction = DataprepActionPtr.Get();
diff --git a/Engine/Plugins/Experimental/DataPrepEditor/Source/DataPrepEditor/Private/Widgets/Action/SDataprepActionSteps.h b/Engine/Plugins/Experimental/DataPrepEditor/Source/DataPrepEditor/Private/Widgets/Action/SDataprepActionSteps.h
--- a/Engine/Plugins/Experimental/DataPrepEditor/Source/DataPrepEditor/Private/Widgets/Action/SDataprepActionSteps.h
+++ b/Engine/Plugins/Experimental/DataPrepEditor/Source/DataPrepEditor/Private/Widgets/Action/SDataprepActionSteps.h
@@ -51,6 +51,8 @@ public:
 
 private:
 
+	void OnStepsOrderChanged();
+
 	// Update the display of the steps
 	void Refresh();
 
diff --git a/Engine/Plugins/Experimental/DataPrepEditor/Source/DataPrepEditor/Private/Widgets/Action/SDataprepFilter.cpp b/Engine/Plugins/Experimental/DataPrepEditor/Source/DataPrepEditor/Private/Widgets/Action/SDataprepFilter.cpp
--- a/Engine/Plugins/Experimental/DataPrepEditor/Source/DataPrepEditor/Private/Widgets/Action/SDataprepFilter.cpp
+++ b/Engine/Plugins/Experimental/DataPrepEditor/Source/DataPrepEditor/Private/Widgets/Action/SDataprepFilter.cpp
@@ -3,6 +3,7 @@
 #include "Widgets/Action/SDataprepFilter.h"
 
 #include "DataprepEditorStyle.h"
+#include "DataprepEditorUtils.h"
 #include "SelectionSystem/DataprepBoolFilter.h"
 #include "SelectionSystem/DataprepFilter.h"
 #include "SelectionSystem/DataprepFloatFilter.h"
@@ -154,6 +155,7 @@ void SDataprepFilter::InverseFilter()
 	{
 		FScopedTransaction Transaction( LOCTEXT("InverseFilterTransaction", "Inverse the filter") );
 		Filter->SetIsExcludingResult( !Filter->IsExcludingResult() );
+		FDataprepEditorUtils::NotifySystemOfChangeInPipeline( Filter );
 	}
 }
 
diff --git a/Engine/Plugins/Experimental/DataPrepEditor/Source/DataPrepEditor/Private/Widgets/Action/SDataprepFloatFilter.cpp b/Engine/Plugins/Experimental/DataPrepEditor/Source/DataPrepEditor/Private/Widgets/Action/SDataprepFloatFilter.cpp
--- a/Engine/Plugins/Experimental/DataPrepEditor/Source/DataPrepEditor/Private/Widgets/Action/SDataprepFloatFilter.cpp
+++ b/Engine/Plugins/Experimental/DataPrepEditor/Source/DataPrepEditor/Private/Widgets/Action/SDataprepFloatFilter.cpp
@@ -2,6 +2,7 @@
 
 #include "Widgets/Action/SDataprepFloatFilter.h"
 
+#include "DataprepEditorUtils.h"
 #include "SchemaActions/DataprepFetcherMenuActionCollector.h"
 #include "SchemaActions/DataprepSchemaAction.h"
 #include "SelectionSystem/DataprepFloatFilter.h"
@@ -131,6 +132,7 @@ void SDataprepFloatFilter::OnSelectedCriteriaChanged(TSharedPtr<FListEntry> List
 	{
 		FScopedTransaction Transaction( LOCTEXT("SelectionCriteriaChangedTransaction", "Changed the Float Selection Criteria") );
 		Filter->SetFloatMatchingCriteria( FloatMatchType );
+		FDataprepEditorUtils::NotifySystemOfChangeInPipeline( Filter );
 	}
 }
 
@@ -181,6 +183,7 @@ void SDataprepFloatFilter::OnEqualValueComitted(float NewEqualValue, ETextCommit
 		Filter->SetEqualValue( OldEqualValue );
 		FScopedTransaction Transaction( LOCTEXT("EqualValueChangedTransaction","Change the Equal Value") );
 		Filter->SetEqualValue( NewEqualValue );
+		FDataprepEditorUtils::NotifySystemOfChangeInPipeline( Filter );
 		OldEqualValue = NewEqualValue;
 	}
 }
@@ -213,6 +216,7 @@ void SDataprepFloatFilter::OnToleranceComitted(float NewTolerance, ETextCommit::
 		Filter->SetTolerance( OldTolerance );
 		FScopedTransaction Transaction( LOCTEXT("ToleranceChangedTransaction", "Change the Tolerance") );
 		Filter->SetTolerance( NewTolerance );
+		FDataprepEditorUtils::NotifySystemOfChangeInPipeline( Filter );
 		OldTolerance = NewTolerance;
 	}
 }
diff --git a/Engine/Plugins/Experimental/DataPrepEditor/Source/DataPrepEditor/Private/Widgets/Action/SDataprepStringFilter.cpp b/Engine/Plugins/Experimental/DataPrepEditor/Source/DataPrepEditor/Private/Widgets/Action/SDataprepStringFilter.cpp
--- a/Engine/Plugins/Experimental/DataPrepEditor/Source/DataPrepEditor/Private/Widgets/Action/SDataprepStringFilter.cpp
+++ b/Engine/Plugins/Experimental/DataPrepEditor/Source/DataPrepEditor/Private/Widgets/Action/SDataprepStringFilter.cpp
@@ -2,6 +2,7 @@
 
 #include "Widgets/Action/SDataprepStringFilter.h"
 
+#include "DataprepEditorUtils.h"
 #include "SchemaActions/DataprepFetcherMenuActionCollector.h"
 #include "SchemaActions/DataprepSchemaAction.h"
 #include "SelectionSystem/DataprepStringFilter.h"
@@ -116,6 +117,7 @@ void SDataprepStringFilter::OnSelectedCriteriaChanged(TSharedPtr<FListEntry> Lis
 	{	
 		FScopedTransaction Transaction( LOCTEXT("SelectionCriteriaChangedTransaction","Changed the String Selection Criteria") );
 		Filter->SetStringMatchingCriteria( StringMatchType );
+		FDataprepEditorUtils::NotifySystemOfChangeInPipeline( Filter );
 	}
 }
 
@@ -140,6 +142,7 @@ void SDataprepStringFilter::OnUserStringComitted(const FText& NewText, ETextComm
 		Filter->SetUserString( OldUserString );
 		FScopedTransaction Transaction( LOCTEXT("SelectionStringChangedTransaction","Changed the Selection String") );
 		Filter->SetUserString( NewUserString );
+		FDataprepEditorUtils::NotifySystemOfChangeInPipeline( Filter );
 		OldUserString = NewUserString;
 	}
 }
diff --git a/Engine/Plugins/Experimental/DataPrepEditor/Source/DataPrepEditor/Private/Widgets/DataprepWidgets.cpp b/Engine/Plugins/Experimental/DataPrepEditor/Source/DataPrepEditor/Private/Widgets/DataprepWidgets.cpp
--- a/Engine/Plugins/Experimental/DataPrepEditor/Source/DataPrepEditor/Private/Widgets/DataprepWidgets.cpp
+++ b/Engine/Plugins/Experimental/DataPrepEditor/Source/DataPrepEditor/Private/Widgets/DataprepWidgets.cpp
@@ -3,6 +3,7 @@
 #include "DataprepWidgets.h"
 
 #include "DataPrepContentConsumer.h"
+#include "DataprepEditorUtils.h"
 
 #include "ContentBrowserModule.h"
 #include "DetailLayoutBuilder.h"
@@ -298,6 +299,8 @@ TSharedRef< SWidget > SDataprepDetailsView::CreateDefaultWidget( TSharedPtr< SWi
 
 void SDataprepDetailsView::OnPropertyChanged(const FPropertyChangedEvent& InEvent)
 {
+	FDataprepEditorUtils::NotifySystemOfChangeInPipeline( DetailedObject );
+
 	if( TrackedProperties.Contains( InEvent.Property ) )
 	{
 		TArray< UObject* > Objects;
@@ -466,8 +469,6 @@ void SDataprepDetailsView::Construct(const FArguments& InArgs)
 	FPropertyRowGeneratorArgs Args;
 	Generator = PropertyEditorModule.CreatePropertyRowGenerator(Args);
 
-	OnPropertyChangedHandle = Generator->OnFinishedChangingProperties().AddSP( this, &SDataprepDetailsView::OnPropertyChanged );
-
 	if( DetailedObject != nullptr )
 	{
 		TArray< UObject* > Objects;
@@ -479,6 +480,7 @@ void SDataprepDetailsView::Construct(const FArguments& InArgs)
 	ColumnSizeData.LeftColumnWidth = TAttribute<float>( this, &SDataprepDetailsView::OnGetLeftColumnWidth );
 	ColumnSizeData.RightColumnWidth = TAttribute<float>( this, &SDataprepDetailsView::OnGetRightColumnWidth );
 	ColumnSizeData.OnWidthChanged = SSplitter::FOnSlotResized::CreateSP( this, &SDataprepDetailsView::OnSetColumnWidth );
+	OnPropertyChangedHandle = Generator->OnFinishedChangingProperties().AddSP( this, &SDataprepDetailsView::OnPropertyChanged );
 
 	Construct();
 }
diff --git a/Engine/Plugins/Experimental/DataPrepEditor/Source/DataPrepLibraries/Private/DataprepOperations.cpp b/Engine/Plugins/Experimental/DataPrepEditor/Source/DataPrepLibraries/Private/DataprepOperations.cpp
--- a/Engine/Plugins/Experimental/DataPrepEditor/Source/DataPrepLibraries/Private/DataprepOperations.cpp
+++ b/Engine/Plugins/Experimental/DataPrepEditor/Source/DataPrepLibraries/Private/DataprepOperations.cpp
@@ -255,9 +255,9 @@ void UDataprepRemoveObjectsOperation::OnExecution_Implementation(const FDataprep
 void FDataprepSetLOGGroupDetails::OnLODGroupChanged(TSharedPtr<FString> NewValue, ESelectInfo::Type /*SelectInfo*/)
 {
 	int32 Index = LODGroupOptions.Find(NewValue);
-	if (Index != INDEX_NONE)
+	if (Index != INDEX_NONE && LodGroupPropertyHandle.IsValid() )
 	{
-		DataprepOperation->GroupName = LODGroupNames[Index];
+		LodGroupPropertyHandle->SetValue( LODGroupNames[Index] );
 	}
 }
 
@@ -310,6 +310,8 @@ void FDataprepSetLOGGroupDetails::CustomizeDetails(IDetailLayoutBuilder & Detail
 	FName CategoryName = CategoryNames.Num() > 0 ? CategoryNames[0] : FName( TEXT("SetLOGGroup_Internal") );
 	IDetailCategoryBuilder& ImportSettingsCategoryBuilder = DetailBuilder.EditCategory( CategoryName, FText::GetEmpty(), ECategoryPriority::Important );
 
+	LodGroupPropertyHandle = DetailBuilder.GetProperty(GET_MEMBER_NAME_CHECKED(UDataprepSetLODGroupOperation, GroupName));
+
 	// Hide GroupName property as it is replaced with custom widget
 	DetailBuilder.HideProperty(GET_MEMBER_NAME_CHECKED(UDataprepSetLODGroupOperation, GroupName));
 
diff --git a/Engine/Plugins/Experimental/DataPrepEditor/Source/DataPrepLibraries/Private/DataprepOperations.h b/Engine/Plugins/Experimental/DataPrepEditor/Source/DataPrepLibraries/Private/DataprepOperations.h
--- a/Engine/Plugins/Experimental/DataPrepEditor/Source/DataPrepLibraries/Private/DataprepOperations.h
+++ b/Engine/Plugins/Experimental/DataPrepEditor/Source/DataPrepLibraries/Private/DataprepOperations.h
@@ -11,6 +11,7 @@
 #include "EditorStaticMeshLibrary.h"
 #include "Engine/EngineTypes.h"
 #include "IDetailCustomization.h"
+#include "PropertyHandle.h"
 
 class SWidget;
 class UMaterialInterface;
@@ -346,4 +347,6 @@ private:
 	/** LOD group options. */
 	TArray< TSharedPtr< FString > > LODGroupOptions;
 	TArray<FName>					LODGroupNames;
+
+	TSharedPtr<IPropertyHandle> LodGroupPropertyHandle;
 };
