Fixed crash when executing graph with LOD operation
Updated the tooltip for Import and Commit buttons
Fixed a logic issue in snapshot code

diff --git a/Engine/Plugins/Experimental/DataPrepEditor/Source/DataPrepEditor/Private/DataPrepEditor.cpp b/Engine/Plugins/Experimental/DataPrepEditor/Source/DataPrepEditor/Private/DataPrepEditor.cpp
--- a/Engine/Plugins/Experimental/DataPrepEditor/Source/DataPrepEditor/Private/DataPrepEditor.cpp
+++ b/Engine/Plugins/Experimental/DataPrepEditor/Source/DataPrepEditor/Private/DataPrepEditor.cpp
@@ -614,12 +614,30 @@ void FDataprepEditor::CleanPreviewWorld()
 			{
 				ObjectToDelete->Rename( nullptr, GetTransientPackage(), REN_DontCreateRedirectors | REN_NonTransactional );
 				ObjectsToDelete.Add( ObjectToDelete );
+
+				// Remove geometry from static meshes to be deleted to avoid unwanted rebuild done when calling FDataprepCoreUtils::PurgeObjects
+				// #ueent_todo: This is a temporary solution. Need to find a better way to do that
+				if( UStaticMesh* StaticMesh = Cast<UStaticMesh>( ObjectToDelete ) )
+				{
+					StaticMesh->ReleaseResources();
+					StaticMesh->ClearMeshDescriptions();
+					StaticMesh->SourceModels.Empty();
+					StaticMesh->StaticMaterials.Empty();
+				}
 			}
 		}
 	}
 
+	// #ueent_todo: Should we find a better way to silently delete assets?
+	// Disable warnings from LogStaticMesh because FDataprepCoreUtils::PurgeObjects is pretty verbose on harmless warnings
+	ELogVerbosity::Type PrevLogStaticMeshVerbosity = LogStaticMesh.GetVerbosity();
+	LogStaticMesh.SetVerbosity( ELogVerbosity::Error );
+
 	FDataprepCoreUtils::PurgeObjects( MoveTemp( ObjectsToDelete ) );
 
+	// Restore LogStaticMesh verbosity
+	LogStaticMesh.SetVerbosity( PrevLogStaticMeshVerbosity );
+
 	CachedAssets.Reset();
 	Assets.Reset();
 
diff --git a/Engine/Plugins/Experimental/DataPrepEditor/Source/DataPrepEditor/Private/DataPrepEditorActions.cpp b/Engine/Plugins/Experimental/DataPrepEditor/Source/DataPrepEditor/Private/DataPrepEditorActions.cpp
--- a/Engine/Plugins/Experimental/DataPrepEditor/Source/DataPrepEditor/Private/DataPrepEditorActions.cpp
+++ b/Engine/Plugins/Experimental/DataPrepEditor/Source/DataPrepEditor/Private/DataPrepEditorActions.cpp
@@ -12,9 +12,9 @@ void FDataprepEditorCommands::RegisterCommands()
 
 	UI_COMMAND(SaveScene, "Save Datasmith Scene", "To be filled.", EUserInterfaceActionType::ToggleButton, FInputChord());
 	UI_COMMAND(ShowDatasmithSceneSettings, "Scene Settings", "Edit the scene settings.", EUserInterfaceActionType::ToggleButton, FInputChord());
-	UI_COMMAND(BuildWorld, "Import", "Load external data sources referenced by producers into the Dataprep Editor", EUserInterfaceActionType::ToggleButton, FInputChord());
+	UI_COMMAND(BuildWorld, "Import", "Load inputs' data into the Dataprep Editor", EUserInterfaceActionType::ToggleButton, FInputChord());
 	UI_COMMAND(ExecutePipeline, "Execute", "Execute the Dataprep graph on the imported data", EUserInterfaceActionType::ToggleButton, FInputChord());
-	UI_COMMAND(CommitWorld, "Commit", "Export the transformed imported data into the target defined by the consumer", EUserInterfaceActionType::ToggleButton, FInputChord());
+	UI_COMMAND(CommitWorld, "Commit", "Copy data from the Dataprep Editor to the project", EUserInterfaceActionType::ToggleButton, FInputChord());
 }
 
 #undef LOCTEXT_NAMESPACE
diff --git a/Engine/Plugins/Experimental/DataPrepEditor/Source/DataPrepEditor/Private/DataprepSnapshot.cpp b/Engine/Plugins/Experimental/DataPrepEditor/Source/DataPrepEditor/Private/DataprepSnapshot.cpp
--- a/Engine/Plugins/Experimental/DataPrepEditor/Source/DataPrepEditor/Private/DataprepSnapshot.cpp
+++ b/Engine/Plugins/Experimental/DataPrepEditor/Source/DataPrepEditor/Private/DataprepSnapshot.cpp
@@ -71,19 +71,22 @@ namespace DataprepSnapshotUtil
 		// Helper struct to identify dependency of a UObject on other UObject(s) except given one (its outer)
 		struct FObjectDependencyAnalyzer : public FArchiveUObject
 		{
-			FObjectDependencyAnalyzer(const TSet<UObject*>& InValidObjects)
-				: ValidObjects( DependentObjects )
+			FObjectDependencyAnalyzer(UObject* InSourceObject, const TSet<UObject*>& InValidObjects)
+				: SourceObject( InSourceObject )
+				, ValidObjects( InValidObjects )
 			{ }
 
 			virtual FArchive& operator<<(UObject*& Obj) override
 			{
 				if(Obj != nullptr)
 				{
-					if( Obj->IsA<UPackage>() || (Obj->HasAnyFlags(RF_Public) && Obj->GetOuter()->IsA<UPackage>()))
+					// Limit serialization to sub-object of source object
+					if( Obj == SourceObject->GetOuter() || Obj->IsA<UPackage>() || (Obj->HasAnyFlags(RF_Public) && Obj->GetOuter()->IsA<UPackage>()))
 					{
 						return FArchiveUObject::operator<<(Obj);
 					}
-					else if( ValidObjects.Contains( Obj ) )
+					// Stop serialization when a dependency is found or has been found
+					else if( Obj != SourceObject && !DependentObjects.Contains( Obj ) && ValidObjects.Contains( Obj ) )
 					{
 						DependentObjects.Add( Obj );
 					}
@@ -92,8 +95,9 @@ namespace DataprepSnapshotUtil
 				return *this;
 			}
 
-			TSet<UObject*> DependentObjects;
+			UObject* SourceObject;
 			const TSet<UObject*>& ValidObjects;
+			TSet<UObject*> DependentObjects;
 		};
 
 		TArray<uint8> MemoryBuffer;
@@ -120,7 +124,7 @@ namespace DataprepSnapshotUtil
 			TSet<UObject*> SubObjectsSet(SubObjectsArray);
 			for(UObject* SubObject : SubObjectsArray)
 			{
-				FObjectDependencyAnalyzer Analyzer( SubObjectsSet );
+				FObjectDependencyAnalyzer Analyzer( SubObject, SubObjectsSet );
 				SubObject->Serialize( Analyzer );
 
 				SubObjectDependencyGraph[SubObject].Append(Analyzer.DependentObjects);
diff --git a/Engine/Plugins/Experimental/DataPrepEditor/Source/DataPrepLibraries/Private/DataPrepOperationsLibrary.cpp b/Engine/Plugins/Experimental/DataPrepEditor/Source/DataPrepLibraries/Private/DataPrepOperationsLibrary.cpp
--- a/Engine/Plugins/Experimental/DataPrepEditor/Source/DataPrepLibraries/Private/DataPrepOperationsLibrary.cpp
+++ b/Engine/Plugins/Experimental/DataPrepEditor/Source/DataPrepLibraries/Private/DataPrepOperationsLibrary.cpp
@@ -17,12 +17,16 @@
 #include "Engine/World.h"
 #include "EngineUtils.h"
 #include "GameFramework/Actor.h"
+#include "IMeshBuilderModule.h"
+#include "Interfaces/ITargetPlatform.h"
+#include "Interfaces/ITargetPlatformManagerModule.h"
 #include "Layers/ILayers.h"
 #include "Materials/Material.h"
 #include "Materials/MaterialInterface.h"
 #include "Math/Vector2D.h"
 #include "MeshAttributeArray.h"
 #include "MeshDescription.h"
+#include "MeshDescriptionOperations.h"
 #include "MeshTypes.h"
 #include "Misc/FileHelper.h"
 #include "ObjectTools.h"
@@ -151,6 +155,8 @@ namespace DataprepOperationsLibraryUtil
 				SourceModel.BuildSettings.bGenerateLightmapUVs = false;
 				SourceModel.BuildSettings.bRecomputeNormals = false;
 				SourceModel.BuildSettings.bRecomputeTangents = false;
+				SourceModel.BuildSettings.bBuildAdjacencyBuffer = false;
+				SourceModel.BuildSettings.bBuildReversedIndexBuffer = false;
 			}
 
 			return BuildSettingsBackup;
@@ -176,6 +182,8 @@ namespace DataprepOperationsLibraryUtil
 					BuildSettings.bGenerateLightmapUVs = CachedBuildSettings.bGenerateLightmapUVs;
 					BuildSettings.bRecomputeNormals = CachedBuildSettings.bRecomputeNormals;
 					BuildSettings.bRecomputeTangents = CachedBuildSettings.bRecomputeTangents;
+					BuildSettings.bBuildAdjacencyBuffer = CachedBuildSettings.bBuildAdjacencyBuffer;
+					BuildSettings.bBuildReversedIndexBuffer = CachedBuildSettings.bBuildReversedIndexBuffer;
 				}
 			}
 		}
@@ -215,6 +223,42 @@ namespace DataprepOperationsLibraryUtil
 			}
 		}
 	}
+
+	// Replacement of UStaticMesh::CacheDerivedData() which performs too much operations for our purpose
+	// And displays unwanted progress bar
+	// #ueent_todo: Work with Geometry team to find the proper replacement
+	void BuildRenderData( UStaticMesh* StaticMesh, IMeshBuilderModule& MeshBuilderModule )
+	{
+		FStaticMeshSourceModel& SourceModel = StaticMesh->SourceModels[0];
+
+		FMeshBuildSettings PrevBuildSettings = SourceModel.BuildSettings;
+		SourceModel.BuildSettings.bGenerateLightmapUVs = false;
+		SourceModel.BuildSettings.bRecomputeNormals = false;
+		SourceModel.BuildSettings.bRecomputeTangents = false;
+		SourceModel.BuildSettings.bBuildAdjacencyBuffer = false;
+		SourceModel.BuildSettings.bBuildReversedIndexBuffer = false;
+
+		FMeshDescription* StaticMeshDescription = SourceModel.MeshDescription.Get();
+		check( StaticMeshDescription );
+
+		// Create render data
+		StaticMesh->RenderData.Reset( new(FMemory::Malloc(sizeof(FStaticMeshRenderData)))FStaticMeshRenderData() );
+
+		ITargetPlatformManagerModule& TargetPlatformManager = GetTargetPlatformManagerRef();
+		ITargetPlatform* RunningPlatform = TargetPlatformManager.GetRunningTargetPlatform();
+		check(RunningPlatform);
+		const FStaticMeshLODSettings& LODSettings = RunningPlatform->GetStaticMeshLODSettings();
+
+		const FStaticMeshLODGroup& LODGroup = LODSettings.GetLODGroup(StaticMesh->LODGroup);
+
+		if ( !MeshBuilderModule.BuildMesh( *StaticMesh->RenderData, StaticMesh, LODGroup ) )
+		{
+			UE_LOG(LogDataprep, Error, TEXT("Failed to build static mesh. See previous line(s) for details."));
+			return;
+		}
+
+		SourceModel.BuildSettings = PrevBuildSettings;
+	}
 } // ns DataprepOperationsLibraryUtil
 
 void UDataprepOperationsLibrary::SetLods(const TArray<UObject*>& SelectedObjects, const FEditorScriptingMeshReductionOptions& ReductionOptions)
@@ -237,6 +281,8 @@ void UDataprepOperationsLibrary::SetSimpleCollision(const TArray<UObject*>& Sele
 {
 	TArray<UStaticMesh*> SelectedMeshes = DataprepOperationsLibraryUtil::GetSelectedMeshes(SelectedObjects);
 
+	IMeshBuilderModule& MeshBuilderModule = FModuleManager::LoadModuleChecked< IMeshBuilderModule >( TEXT("MeshBuilder") );
+
 	// Create LODs but do not commit changes
 	for (UStaticMesh* StaticMesh : SelectedMeshes)
 	{
@@ -245,9 +291,37 @@ void UDataprepOperationsLibrary::SetSimpleCollision(const TArray<UObject*>& Sele
 			DataprepOperationsLibraryUtil::FScopedStaticMeshEdit StaticMeshEdit( StaticMesh );
 
 			// Remove existing simple collisions
-			UEditorStaticMeshLibrary::RemoveCollisionsWithNotification(StaticMesh, false);
+			UEditorStaticMeshLibrary::RemoveCollisionsWithNotification( StaticMesh, false );
 
-			UEditorStaticMeshLibrary::AddSimpleCollisionsWithNotification(StaticMesh, ShapeType, false);
+			// Check that render data is available if k-DOP type of collision is required
+			bool bFreeRenderData = false;
+			switch (ShapeType)
+			{
+				case EScriptingCollisionShapeType::NDOP10_X:
+				case EScriptingCollisionShapeType::NDOP10_Y:
+				case EScriptingCollisionShapeType::NDOP10_Z:
+				case EScriptingCollisionShapeType::NDOP18:
+				case EScriptingCollisionShapeType::NDOP26:
+				{
+					if( StaticMesh->RenderData == nullptr )
+					{
+						DataprepOperationsLibraryUtil::BuildRenderData( StaticMesh, MeshBuilderModule );
+						bFreeRenderData = true;
+					}
+					break;
+				}
+				default:
+				{
+					break;
+				}
+			}
+
+			UEditorStaticMeshLibrary::AddSimpleCollisionsWithNotification( StaticMesh, ShapeType, false );
+
+			if( bFreeRenderData )
+			{
+				StaticMesh->RenderData = nullptr;
+			}
 		}
 	}
 }
