Revisited the IDataprepProgressReporter interface: Will need to be reviewed for 4.24
Added code to report progress in producers, consumers and while executing the list of actions in the pipeline

diff --git a/Engine/Plugins/Experimental/DataPrepEditor/Source/DataPrepCore/Private/DataPrepAsset.cpp b/Engine/Plugins/Experimental/DataPrepEditor/Source/DataPrepCore/Private/DataPrepAsset.cpp
--- a/Engine/Plugins/Experimental/DataPrepEditor/Source/DataPrepCore/Private/DataPrepAsset.cpp
+++ b/Engine/Plugins/Experimental/DataPrepEditor/Source/DataPrepCore/Private/DataPrepAsset.cpp
@@ -58,7 +58,9 @@ void UDataprepAsset::PostInitProperties()
 			{
 				if( CurrentClass->IsChildOf( UDataprepContentConsumer::StaticClass() ) )
 				{
-					Consumer = NewObject< UDataprepContentConsumer >( this, CurrentClass, NAME_None, RF_Transactional );
+					FString BaseName = GetName() + TEXT("_Consumer");
+					FName ConsumerName = MakeUniqueObjectName( this, CurrentClass, *BaseName );
+					Consumer = NewObject< UDataprepContentConsumer >( this, CurrentClass, ConsumerName, RF_Transactional );
 					check( Consumer );
 
 					FAssetRegistryModule::AssetCreated( Consumer );
@@ -114,21 +116,32 @@ void UDataprepAsset::Serialize( FArchive& Ar )
 
 void UDataprepAsset::RunProducers(const UDataprepContentProducer::ProducerContext& InContext, TArray< TWeakObjectPtr< UObject > >& OutAssets)
 {
+	if( Producers.Num() == 0 )
+	{
+		return;
+	}
+
 	OutAssets.Empty();
 
+	FDataprepProgressTask Task( *InContext.ProgressReporterPtr, NSLOCTEXT( "DataprepAsset", "RunProducers", "Importing ..." ), (float)Producers.Num(), 1.0f );
+
 	for ( FDataprepAssetProducer& AssetProducer : Producers )
 	{
-		const bool bIsOkToRun = AssetProducer.Producer && AssetProducer.bIsEnabled &&
+		UDataprepContentProducer* Producer = AssetProducer.Producer;
+
+		const bool bIsOkToRun = Producer && AssetProducer.bIsEnabled &&
 			( AssetProducer.SupersededBy == INDEX_NONE || !Producers[AssetProducer.SupersededBy].bIsEnabled );
 
+		Task.ReportNextStep( FText::Format( NSLOCTEXT( "DataprepAsset", "ProducerReport", "Importing {0} ..."), FText::FromString( Producer->GetName() ) ) );
+
 		if ( bIsOkToRun )
 		{
 			FString OutReason;
-			if (AssetProducer.Producer->Initialize( InContext, OutReason ))
+			if (Producer->Initialize( InContext, OutReason ))
 			{
-				if ( AssetProducer.Producer->Produce() )
+				if ( Producer->Produce() )
 				{
-					const TArray< TWeakObjectPtr< UObject > >& ProducerAssets = AssetProducer.Producer->GetAssets();
+					const TArray< TWeakObjectPtr< UObject > >& ProducerAssets = Producer->GetAssets();
 					if (ProducerAssets.Num() > 0)
 					{
 						OutAssets.Append( ProducerAssets );
@@ -136,11 +149,11 @@ void UDataprepAsset::RunProducers(const UDataprepContentProducer::ProducerContex
 				}
 				else
 				{
-					OutReason = FText::Format(NSLOCTEXT("DataprepAsset", "ProducerRunFailed", "{0} failed to run."), FText::FromString( AssetProducer.Producer->GetName() ) ).ToString();
+					OutReason = FText::Format(NSLOCTEXT("DataprepAsset", "ProducerRunFailed", "{0} failed to run."), FText::FromString( Producer->GetName() ) ).ToString();
 				}
 			}
 
-			AssetProducer.Producer->Reset();
+			Producer->Reset();
 
 			if( !OutReason.IsEmpty() )
 			{
@@ -312,7 +325,9 @@ bool UDataprepAsset::ReplaceConsumer(UClass* NewConsumerClass)
 			DataprepAssetUtil::DeleteRegisteredAsset( Consumer );
 		}
 
-		Consumer = NewObject< UDataprepContentConsumer >( this, NewConsumerClass, NAME_None, RF_Transactional );
+		FString BaseName = GetName() + TEXT("_Consumer");
+		FName ConsumerName = MakeUniqueObjectName( this, NewConsumerClass, *BaseName );
+		Consumer = NewObject< UDataprepContentConsumer >( this, NewConsumerClass, ConsumerName, RF_Transactional );
 		check( Consumer );
 
 		FAssetRegistryModule::AssetCreated( Consumer );
diff --git a/Engine/Plugins/Experimental/DataPrepEditor/Source/DataPrepCore/Private/DataPrepOperation.cpp b/Engine/Plugins/Experimental/DataPrepEditor/Source/DataPrepCore/Private/DataPrepOperation.cpp
--- a/Engine/Plugins/Experimental/DataPrepEditor/Source/DataPrepCore/Private/DataPrepOperation.cpp
+++ b/Engine/Plugins/Experimental/DataPrepEditor/Source/DataPrepCore/Private/DataPrepOperation.cpp
@@ -50,22 +50,6 @@ void UDataprepOperation::LogError(const FText& InLogText)
 	}
 }
 
-void UDataprepOperation::ReportProgress(float Progress)
-{
-	if ( OperationContext && OperationContext->DataprepProgressReporter )
-	{
-		OperationContext->DataprepProgressReporter->ReportProgress( Progress, *this );
-	}
-}
-
-void UDataprepOperation::ReportProgressWithMessage(float Progress, const FText& InMessage)
-{
-	if ( OperationContext && OperationContext->DataprepProgressReporter )
-	{
-		OperationContext->DataprepProgressReporter->ReportProgressWithMessage( Progress, InMessage, *this );
-	}
-}
-
 void UDataprepOperation::ExecuteOperation(TSharedRef<FDataprepOperationContext>&  InOperationContext)
 {
 	OperationContext = InOperationContext;
diff --git a/Engine/Plugins/Experimental/DataPrepEditor/Source/DataPrepCore/Private/DataprepCoreUtils.cpp b/Engine/Plugins/Experimental/DataPrepEditor/Source/DataPrepCore/Private/DataprepCoreUtils.cpp
--- a/Engine/Plugins/Experimental/DataPrepEditor/Source/DataPrepCore/Private/DataprepCoreUtils.cpp
+++ b/Engine/Plugins/Experimental/DataPrepEditor/Source/DataPrepCore/Private/DataprepCoreUtils.cpp
@@ -2,6 +2,8 @@
 
 #include "DataprepCoreUtils.h"
 
+#include "IDataprepProgressReporter.h"
+
 #include "UObject/UObjectHash.h"
 
 #if WITH_EDITOR
@@ -72,3 +74,19 @@ void FDataprepCoreUtils::PurgeObjects(TArray<UObject*> Objects)
 	}
 }
 
+FDataprepProgressTask::FDataprepProgressTask(IDataprepProgressReporter& InReporter, const FText& InDescription, float InAmountOfWork, float InIncrementOfWork)
+	: Reporter( InReporter )
+	, DefaultIncrementOfWork( InIncrementOfWork )
+{
+	Reporter.PushTask( InDescription, InAmountOfWork );
+}
+
+FDataprepProgressTask::~FDataprepProgressTask()
+{
+	Reporter.PopTask();
+}
+
+void FDataprepProgressTask::ReportNextStep(const FText & InMessage, float InIncrementOfWork )
+{
+	Reporter.ReportProgress( InIncrementOfWork, InMessage );
+}
\ No newline at end of file
diff --git a/Engine/Plugins/Experimental/DataPrepEditor/Source/DataPrepCore/Public/DataPrepContentConsumer.h b/Engine/Plugins/Experimental/DataPrepEditor/Source/DataPrepCore/Public/DataPrepContentConsumer.h
--- a/Engine/Plugins/Experimental/DataPrepEditor/Source/DataPrepCore/Public/DataPrepContentConsumer.h
+++ b/Engine/Plugins/Experimental/DataPrepEditor/Source/DataPrepCore/Public/DataPrepContentConsumer.h
@@ -166,23 +166,7 @@ protected:
 			Context.LoggerPtr->LogError( Message, *this );
 		}
 	}
-
-	void ReportProgress(float Progress)
-	{
-		if ( Context.ProgressReporterPtr.IsValid() )
-		{
-			Context.ProgressReporterPtr->ReportProgress( Progress, *this );
-		}
-	}
-
-	void ReportProgressWithMessage( float Progress, const FText& Message )
-	{
-		if ( Context.ProgressReporterPtr.IsValid() )
-		{
-			Context.ProgressReporterPtr->ReportProgressWithMessage( Progress, Message, *this );
-		}
-	}
-	// End of helper functions to log messages and report progress
+	// End of helper functions to log messages
 
 protected:
 	UPROPERTY( EditAnywhere, Category = DataprepConsumerInternal )
diff --git a/Engine/Plugins/Experimental/DataPrepEditor/Source/DataPrepCore/Public/DataPrepContentProducer.h b/Engine/Plugins/Experimental/DataPrepEditor/Source/DataPrepCore/Public/DataPrepContentProducer.h
--- a/Engine/Plugins/Experimental/DataPrepEditor/Source/DataPrepCore/Public/DataPrepContentProducer.h
+++ b/Engine/Plugins/Experimental/DataPrepEditor/Source/DataPrepCore/Public/DataPrepContentProducer.h
@@ -159,23 +159,7 @@ protected:
 			Context.LoggerPtr->LogError( Message, *this );
 		}
 	}
-
-	void ReportProgress(float Progress)
-	{
-		if ( Context.ProgressReporterPtr.IsValid() )
-		{
-			Context.ProgressReporterPtr->ReportProgress( Progress, *this );
-		}
-	}
-
-	void ReportProgressWithMessage( float Progress, const FText& Message )
-	{
-		if ( Context.ProgressReporterPtr.IsValid() )
-		{
-			Context.ProgressReporterPtr->ReportProgressWithMessage( Progress, Message, *this );
-		}
-	}
-	// End of helper functions to log messages and report progress
+	// End of helper functions to log messages
 
 	bool IsValid() { return Context.WorldPtr.IsValid() && Context.RootPackagePtr.IsValid(); }
 
diff --git a/Engine/Plugins/Experimental/DataPrepEditor/Source/DataPrepCore/Public/DataPrepOperation.h b/Engine/Plugins/Experimental/DataPrepEditor/Source/DataPrepCore/Public/DataPrepOperation.h
--- a/Engine/Plugins/Experimental/DataPrepEditor/Source/DataPrepCore/Public/DataPrepOperation.h
+++ b/Engine/Plugins/Experimental/DataPrepEditor/Source/DataPrepCore/Public/DataPrepOperation.h
@@ -87,21 +87,6 @@ protected:
 	UFUNCTION(BlueprintCallable,  Category = "Log", meta = (HideSelfPin = "true"))
 	void LogError(const FText& InLogError);
 
-	/**
-	 * Report the progress of the operation to the execution context.
-	 * @param Progress The progress of the made since the last call. (1.0 = 100% and 0.0 = 0%)
-	 */
-	UFUNCTION(BlueprintCallable,  Category = "Progress", meta = (HideSelfPin = "true"))
-	void ReportProgress(float Progress);
-
-	/**
-	 * Report the progress of the operation to the execution context
-	 * @param Progress The progress of the made since the last call. (1.0 = 100% and 0.0 = 0%)
-	 * @param Message This allow the user to change an displayed message.
-	 */
-	UFUNCTION(BlueprintCallable, Category = "Progress", meta = (HideSelfPin = "true"))
-	void ReportProgressWithMessage(float Progress, const FText& InMessage);
-
 	// User friendly interface end here ========================================================================
 
 public:
@@ -145,3 +130,17 @@ public:
 private:
 	TSharedPtr<const FDataprepOperationContext> OperationContext;
 };
+
+
+// Todo add class wide comment
+UCLASS(Experimental, Abstract, Blueprintable)
+class DATAPREPCORE_API UDataprepAdditiveOperation : public UDataprepOperation
+{
+	GENERATED_BODY()
+
+public:
+	const TArray< UObject* > GetCreatedAssets() { return CreatedAssets; }
+
+protected:
+	TArray< UObject* > CreatedAssets;
+};
\ No newline at end of file
diff --git a/Engine/Plugins/Experimental/DataPrepEditor/Source/DataPrepCore/Public/IDataprepProgressReporter.h b/Engine/Plugins/Experimental/DataPrepEditor/Source/DataPrepCore/Public/IDataprepProgressReporter.h
--- a/Engine/Plugins/Experimental/DataPrepEditor/Source/DataPrepCore/Public/IDataprepProgressReporter.h
+++ b/Engine/Plugins/Experimental/DataPrepEditor/Source/DataPrepCore/Public/IDataprepProgressReporter.h
@@ -6,25 +6,68 @@ class FText;
 class UObject;
 
 /**
- * This is the interface that a progress reporter must implement to work with the Dataprep core functionalities
+ * This is the interface that a progress reporter must implement to work with FDataprepProgressTask
+ * 
  */
 class IDataprepProgressReporter
 {
 public:
+	virtual ~IDataprepProgressReporter() = default;
+
+protected:
 	/**
-	 * Records the progress of an operation
-	 * @param Progress The progress of the made since the last call. (1.0 = 100% and 0.0 = 0%)
-	 * @param InObject The object that reported the progress
+	 * Push a new task on the stack of started tasks
+	 * @param InDescription		Text to be displayed for the new task
+	 * @param InAmountOfWork	Total amount of work for the new task
 	 */
-	virtual void ReportProgress(float Progress, const UObject& InObject) = 0;
+	virtual void PushTask( const FText& InDescription, float InAmountOfWork ) = 0;
+
+	/** Pop out the current task */
+	virtual void PopTask() = 0;
 
 	/**
-	 * Records the progress of an operation
-	 * @param Progress The progress of the made since the last call. (1.0 = 100% and 0.0 = 0%)
-	 * @param Message This allow the user to change the displayed message.
-	 * @param InObject The object that reported the progress
+	 * Report foreseen progress on the current task
+	 * @param IncrementOfWork	Amount of progress foreseen until the next call
+	 * @param InMessage			Message to be displayed along side the reported progress
 	 */
-	virtual void ReportProgressWithMessage(float Progress, const FText& InMessage, const UObject& InObject) = 0;
+	virtual void ReportProgress( float IncrementOfWork, const FText& InMessage ) = 0;
 
-	virtual ~IDataprepProgressReporter() = default;
+	friend class FDataprepProgressTask;
 };
+
+class DATAPREPCORE_API FDataprepProgressTask
+{
+public:
+	/**
+	 * Report foreseen progress on the current task
+	 * @param InDescription			Description of the task
+	 * @param InAmountOfWork		Total amount of work for the task
+	 * @param InIncrementOfWork		Amount of incremental work at each step within the task
+	 */
+	FDataprepProgressTask( IDataprepProgressReporter& InReporter, const FText& InDescription, float InAmountOfWork, float InIncrementOfWork = 1.0f );
+
+	~FDataprepProgressTask();
+
+	/**
+	 * Report foreseen incremental amount of work until next call
+	 * @param InMessage				Message to be displayed along side the reported progress
+	 * @param InIncrementOfWork		Amount of incremental work foreseen during the next step
+	 */
+	void ReportNextStep( const FText& InMessage, float InIncrementOfWork );
+
+	/**
+	 * Report foreseen default incremental amount of work until next call
+	 * @param InMessage		Message to be displayed along side the reported progress
+	 */
+	void ReportNextStep( const FText& InMessage )
+	{
+		ReportNextStep( InMessage, DefaultIncrementOfWork );
+	}
+
+private:
+	/** Dataprep progress reporter associated with the task */
+	IDataprepProgressReporter& Reporter;
+
+	/** Default incremental amount of work for each step constituting the task */
+	float DefaultIncrementOfWork;
+};
\ No newline at end of file
diff --git a/Engine/Plugins/Experimental/DataPrepEditor/Source/DataPrepEditor/Private/DataPrepEditor.cpp b/Engine/Plugins/Experimental/DataPrepEditor/Source/DataPrepEditor/Private/DataPrepEditor.cpp
--- a/Engine/Plugins/Experimental/DataPrepEditor/Source/DataPrepEditor/Private/DataPrepEditor.cpp
+++ b/Engine/Plugins/Experimental/DataPrepEditor/Source/DataPrepEditor/Private/DataPrepEditor.cpp
@@ -125,33 +125,40 @@ public:
 class FDataprepProgressReporter : public IDataprepProgressReporter
 {
 public:
-	FDataprepProgressReporter( const FText& Title )
+	FDataprepProgressReporter()
 	{
-		ProgressTask = TSharedPtr<FScopedSlowTask>( new FScopedSlowTask( 100.0f, Title, true, *GWarn ) );
-		ProgressTask->MakeDialog(true);
 	}
 
-	virtual ~FDataprepProgressReporter() {}
+	virtual ~FDataprepProgressReporter()
+	{
+	}
 
-	// Begin IDataprepProgressReporter interface
-	virtual void ReportProgress(float Progress, const UObject& InObject) override
+	virtual void PushTask( const FText& InTitle, float InAmountOfWork ) override
 	{
-		ReportProgressWithMessage( Progress, ProgressTask->GetCurrentMessage(), InObject );
+		ProgressTasks.Emplace( new FScopedSlowTask( InAmountOfWork, InTitle, true, *GWarn ) );
+		ProgressTasks.Last()->MakeDialog(true);
 	}
 
-	virtual void ReportProgressWithMessage(float Progress, const FText& InMessage, const UObject& InObject)
+	virtual void PopTask() override
 	{
-		if( ProgressTask.IsValid() )
+		if(ProgressTasks.Num() > 0)
 		{
-			FText ProgressMsg = FText::FromString( FString::Printf( TEXT("%s : %s ..."), *InObject.GetName(), *InMessage.ToString() ) );
-			ProgressTask->EnterProgressFrame( Progress, ProgressMsg );
-			ProgressTask->EnterProgressFrame( 0.0f, ProgressMsg );
+			ProgressTasks.Pop();
+		}
+	}
+
+	// Begin IDataprepProgressReporter interface
+	virtual void ReportProgress( float Progress, const FText& InMessage ) override
+	{
+		if( ProgressTasks.Num() > 0 )
+		{
+			TSharedPtr<FScopedSlowTask>& ProgressTask = ProgressTasks.Last();
+			ProgressTask->EnterProgressFrame( Progress, InMessage );
 		}
 	}
-	// End IDataprepProgressReporter interface
 
 private:
-	TSharedPtr< FScopedSlowTask > ProgressTask;
+	TArray< TSharedPtr< FScopedSlowTask > > ProgressTasks;
 };
 
 FDataprepEditor::FDataprepEditor()
@@ -517,7 +524,7 @@ void FDataprepEditor::OnBuildWorld()
 	Context.SetWorld( PreviewWorld.Get() )
 		.SetRootPackage( TransientPackage )
 		.SetLogger( TSharedPtr< IDataprepLogger >( new FDataprepLogger ) )
-		.SetProgressReporter( TSharedPtr< IDataprepProgressReporter >( new FDataprepProgressReporter( LOCTEXT("Dataprep_BuildWorld", "Importing ...") ) ) );
+		.SetProgressReporter( TSharedPtr< IDataprepProgressReporter >( new FDataprepProgressReporter() ) );
 
 	DataprepAssetPtr->RunProducers( Context, Assets );
 
@@ -656,7 +663,7 @@ void FDataprepEditor::OnExecutePipeline()
 		RestoreFromSnapshot();
 	}
 
-	TSharedPtr< IDataprepProgressReporter > ProgressReporter( new FDataprepProgressReporter( LOCTEXT("Dataprep_ExecutePipeline", "Executing pipeline ...") ) );
+	TSharedPtr< IDataprepProgressReporter > ProgressReporter( new FDataprepProgressReporter() );
 
 	// Trigger execution of data preparation operations on world attached to recipe
 	{
@@ -670,12 +677,31 @@ void FDataprepEditor::OnExecutePipeline()
 		UEdGraphPin* StartNodePin = StartNode->FindPin(UEdGraphSchema_K2::PN_Then, EGPD_Output);
 		if( StartNodePin && StartNodePin->LinkedTo.Num() > 0 )
 		{
-			UEdGraphPin* NextNodeInPin = StartNodePin->LinkedTo[0];
-			while( NextNodeInPin != nullptr )
+			int32 ActionNodeCount = 0;
+			for( UEdGraphPin* NextNodeInPin = StartNodePin->LinkedTo[0]; NextNodeInPin != nullptr ; )
+			{
+				UEdGraphNode* NextNode = NextNodeInPin->GetOwningNode();
+
+				if(UK2Node_DataprepAction* ActionNode = Cast<UK2Node_DataprepAction>(NextNode))
+				{
+					++ActionNodeCount;
+				}
+
+				UEdGraphPin* NextNodeOutPin = NextNode->FindPin( UEdGraphSchema_K2::PN_Then, EGPD_Output );
+				NextNodeInPin = NextNodeOutPin ? ( NextNodeOutPin->LinkedTo.Num() > 0 ? NextNodeOutPin->LinkedTo[0] : nullptr ) : nullptr;
+			}
+
+			ActionNodesExecuted.Reserve( ActionNodeCount );
+
+			FDataprepProgressTask Task( *ProgressReporter, LOCTEXT( "DataprepEditor_ExecutingPipeline", "Executing pipeline ..." ), (float)ActionNodeCount, 1.0f );
+
+			for( UEdGraphPin* NextNodeInPin = StartNodePin->LinkedTo[0]; NextNodeInPin != nullptr ; )
 			{
 				UEdGraphNode* NextNode = NextNodeInPin->GetOwningNode();
 				if( UK2Node_DataprepAction* ActionNode = Cast<UK2Node_DataprepAction>( NextNode ) )
 				{
+					Task.ReportNextStep( FText::Format( LOCTEXT( "DataprepEditor_ExecutingAction", "Executing \"{0}\" ..."), ActionNode->GetNodeTitle( ENodeTitleType::FullTitle ) ) );
+
 					// Break the loop if the node had already been executed
 					if( ActionNodesExecuted.Find( ActionNode ) )
 					{
@@ -792,7 +818,7 @@ void FDataprepEditor::OnCommitWorld()
 		.SetAssets( ValidAssets )
 		.SetTransientContentFolder( GetTransientContentFolder() )
 		.SetLogger( TSharedPtr<IDataprepLogger>( new FDataprepLogger ) )
-		.SetProgressReporter( TSharedPtr< IDataprepProgressReporter >( new FDataprepProgressReporter( LOCTEXT("Dataprep_CommitWorld", "Committing ...") ) ) );
+		.SetProgressReporter( TSharedPtr< IDataprepProgressReporter >( new FDataprepProgressReporter() ) );
 
 	FString OutReason;
 	if( !DataprepAssetPtr->RunConsumer( Context, OutReason ) )
