Fixed static analysis (blind) warning

diff --git a/Engine/Plugins/Experimental/DataPrepEditor/Source/DataPrepCore/Private/DataPrepAsset.cpp b/Engine/Plugins/Experimental/DataPrepEditor/Source/DataPrepCore/Private/DataPrepAsset.cpp
--- a/Engine/Plugins/Experimental/DataPrepEditor/Source/DataPrepCore/Private/DataPrepAsset.cpp
+++ b/Engine/Plugins/Experimental/DataPrepEditor/Source/DataPrepCore/Private/DataPrepAsset.cpp
@@ -127,39 +127,44 @@ void UDataprepAsset::RunProducers(const UDataprepContentProducer::ProducerContex
 
 	for ( FDataprepAssetProducer& AssetProducer : Producers )
 	{
-		UDataprepContentProducer* Producer = AssetProducer.Producer;
-
-		const bool bIsOkToRun = Producer && AssetProducer.bIsEnabled &&
-			( AssetProducer.SupersededBy == INDEX_NONE || !Producers[AssetProducer.SupersededBy].bIsEnabled );
+		if( UDataprepContentProducer* Producer = AssetProducer.Producer )
+		{
+			Task.ReportNextStep( FText::Format( NSLOCTEXT( "DataprepAsset", "ProducerReport", "Importing {0} ..."), FText::FromString( Producer->GetName() ) ) );
 
-		Task.ReportNextStep( FText::Format( NSLOCTEXT( "DataprepAsset", "ProducerReport", "Importing {0} ..."), FText::FromString( Producer->GetName() ) ) );
+			// Run producer if enabled and, if superseded, superseder is disabled
+			const bool bIsOkToRun = AssetProducer.bIsEnabled &&	( AssetProducer.SupersededBy == INDEX_NONE || !Producers[AssetProducer.SupersededBy].bIsEnabled );
 
-		if ( bIsOkToRun )
-		{
-			FString OutReason;
-			if (Producer->Initialize( InContext, OutReason ))
+			if ( bIsOkToRun )
 			{
-				if ( Producer->Produce() )
+				FString OutReason;
+				if (Producer->Initialize( InContext, OutReason ))
 				{
-					const TArray< TWeakObjectPtr< UObject > >& ProducerAssets = Producer->GetAssets();
-					if (ProducerAssets.Num() > 0)
+					if ( Producer->Produce() )
 					{
-						OutAssets.Append( ProducerAssets );
+						const TArray< TWeakObjectPtr< UObject > >& ProducerAssets = Producer->GetAssets();
+						if (ProducerAssets.Num() > 0)
+						{
+							OutAssets.Append( ProducerAssets );
+						}
+					}
+					else
+					{
+						OutReason = FText::Format(NSLOCTEXT("DataprepAsset", "ProducerRunFailed", "{0} failed to run."), FText::FromString( Producer->GetName() ) ).ToString();
 					}
 				}
-				else
-				{
-					OutReason = FText::Format(NSLOCTEXT("DataprepAsset", "ProducerRunFailed", "{0} failed to run."), FText::FromString( Producer->GetName() ) ).ToString();
-				}
-			}
 
-			Producer->Reset();
+				Producer->Reset();
 
-			if( !OutReason.IsEmpty() )
-			{
-				// #ueent_todo: Log that producer has failed
+				if( !OutReason.IsEmpty() )
+				{
+					// #ueent_todo: Log that producer has failed
+				}
 			}
 		}
+		else
+		{
+			Task.ReportNextStep( NSLOCTEXT( "DataprepAsset", "ProducerReport", "Skipped invalid producer ...") );
+		}
 	}
 }
 
