Disabled Dataprep Editor plugin by default
Fixed (non-reported) issue with restoring textures from snapshot

diff --git a/Engine/Plugins/Experimental/DataPrepEditor/DataPrepEditor.uplugin b/Engine/Plugins/Experimental/DataPrepEditor/DataPrepEditor.uplugin
--- a/Engine/Plugins/Experimental/DataPrepEditor/DataPrepEditor.uplugin
+++ b/Engine/Plugins/Experimental/DataPrepEditor/DataPrepEditor.uplugin
@@ -10,7 +10,7 @@
 	"DocsURL": "",
 	"MarketplaceURL": "",
 	"SupportURL": "",
-	"EnabledByDefault": true,
+	"EnabledByDefault": false,
 	"CanContainContent": true,
 	"IsBetaVersion": true,
 	"Installed": false,
diff --git a/Engine/Plugins/Experimental/DataPrepEditor/Source/DataPrepEditor/Private/DataprepSnapshot.cpp b/Engine/Plugins/Experimental/DataPrepEditor/Source/DataPrepEditor/Private/DataprepSnapshot.cpp
--- a/Engine/Plugins/Experimental/DataPrepEditor/Source/DataPrepEditor/Private/DataprepSnapshot.cpp
+++ b/Engine/Plugins/Experimental/DataPrepEditor/Source/DataPrepEditor/Private/DataprepSnapshot.cpp
@@ -2,6 +2,7 @@
 
 #include "DataPrepEditor.h"
 
+#include "DataprepCoreUtils.h"
 #include "DataprepEditorLogCategory.h"
 
 #include "ActorEditorUtils.h"
@@ -9,6 +10,7 @@
 #include "Async/ParallelFor.h"
 #include "Editor/UnrealEdEngine.h"
 #include "EngineUtils.h"
+#include "Engine/Texture.h"
 #include "Exporters/Exporter.h"
 #include "Factories/LevelFactory.h"
 #include "HAL/FileManager.h"
@@ -194,6 +196,12 @@ namespace DataprepSnapshotUtil
 		// Serialize object
 		Object->Serialize(Ar);
 
+		if(UTexture* Texture = Cast<UTexture>(Object))
+		{
+			bool bRebuildResource = !!Texture->Resource;
+			Ar << bRebuildResource;
+		}
+
 		if(bUseCompression)
 		{
 			const int32 BufferHeaderSize = (int32)sizeof(int32);
@@ -215,7 +223,7 @@ namespace DataprepSnapshotUtil
 		}
 	}
 
-	void ReadSnapshotData(UObject* Object, const TArray<uint8>& InSerializedData, TMap<FString, UClass*>& InClassesMap)
+	void ReadSnapshotData(UObject* Object, const TArray<uint8>& InSerializedData, TMap<FString, UClass*>& InClassesMap, TArray<UObject*>& ObjectsToDelete)
 	{
 		TArray<uint8> MemoryBuffer;
 		if(bUseCompression)
@@ -253,7 +261,6 @@ namespace DataprepSnapshotUtil
 			UClass** SubObjectClassPtr = InClassesMap.Find(ClassName);
 			check( SubObjectClassPtr );
 
-			// #ueent_todo: TBF: outer will not be properly updated if sub-object depends on another sub-object 
 			UObject* SubObject = NewObject<UObject>( Object, *SubObjectClassPtr, NAME_None, RF_Transient );
 			SubObjectsArray.Add( SubObject );
 		}
@@ -278,10 +285,42 @@ namespace DataprepSnapshotUtil
 		for(UObject* SubObject : SubObjectsArray)
 		{
 			SubObject->Serialize(Ar);
+
+			// Some sub-objects might have been created by default when their owner is created.
+			// Duplicates must be deleted before loading the owner
+			TArray< UObject* > SiblingObjectsArray;
+			GetObjectsWithOuter( SubObject->GetOuter(), SiblingObjectsArray, /*bIncludeNestedObjects = */ true );
+			if(SiblingObjectsArray.Num() > 1)
+			{
+				ObjectsToDelete.Reserve( ObjectsToDelete.Num() + SiblingObjectsArray.Num() );
+
+				const FName SubObjectName = SubObject->GetFName();
+				const UClass* SubObjectClass = SubObject->GetClass();
+
+				for(UObject* SiblingObject : SiblingObjectsArray)
+				{
+					if(SiblingObject != SubObject && SiblingObject->GetFName() == SubObjectName && SiblingObject->GetClass() == SubObjectClass)
+					{
+						SiblingObject->Rename( nullptr, GetTransientPackage(), REN_DontCreateRedirectors | REN_NonTransactional );
+						ObjectsToDelete.Add( SiblingObject );
+					}
+				}
+			}
 		}
 
 		// Deserialize object
 		Object->Serialize(Ar);
+
+		if(UTexture* Texture = Cast<UTexture>(Object))
+		{
+			bool bRebuildResource = false;
+			Ar << bRebuildResource;
+
+			if(bRebuildResource)
+			{
+				Texture->UpdateResource();
+			}
+		}
 	}
 }
 
@@ -571,6 +610,8 @@ void FDataprepEditor::RestoreFromSnapshot()
 			PackagesCreated.Add( PackageToLoadPath, PackageCreated );
 		}
 
+		// Duplicate sub-objects to delete after all assets are read
+		TArray<UObject*> ObjectsToDelete;
 		UPackage* Package = PackagesCreated[PackageToLoadPath];
 		UObject* Asset = NewObject<UObject>( Package, DataEntry.Get<1>(), *AssetName, DataEntry.Get<2>() );
 
@@ -579,9 +620,11 @@ void FDataprepEditor::RestoreFromSnapshot()
 			FString AssetFilePath = DataprepSnapshotUtil::BuildAssetFileName( TempDir, ObjectPath.GetAssetPathString() );
 			FFileHelper::LoadFileToArray( SerializedData, *AssetFilePath );
 
-			DataprepSnapshotUtil::ReadSnapshotData( Asset, SerializedData, SnapshotClassesMap );
+			DataprepSnapshotUtil::ReadSnapshotData( Asset, SerializedData, SnapshotClassesMap, ObjectsToDelete );
 		}
 
+		FDataprepCoreUtils::PurgeObjects( MoveTemp( ObjectsToDelete ) );
+
 		Assets.Add( Asset );
 
 		UE_LOG( LogDataprepEditor, Verbose, TEXT("Asset %s loaded"), *ObjectPath.GetAssetPathString() );
