Fixed update of outer when reading a snapshot back

diff --git a/Engine/Plugins/Experimental/DataPrepEditor/Source/DataPrepEditor/Private/DataprepSnapshot.cpp b/Engine/Plugins/Experimental/DataPrepEditor/Source/DataPrepEditor/Private/DataprepSnapshot.cpp
--- a/Engine/Plugins/Experimental/DataPrepEditor/Source/DataPrepEditor/Private/DataprepSnapshot.cpp
+++ b/Engine/Plugins/Experimental/DataPrepEditor/Source/DataPrepEditor/Private/DataprepSnapshot.cpp
@@ -179,9 +179,11 @@ namespace DataprepSnapshotUtil
 		}
 
 		// Serialize sub-objects' outer path
-		for(UObject* SubObject : SubObjectsArray)
+		// Done in reverse order since an object can be the outer of the object
+		// it depends on. Not the opposite
+		for(int32 Index = SubObjectsArray.Num() - 1; Index >= 0; --Index)
 		{
-			FSoftObjectPath SoftPath( SubObject->GetOuter() );
+			FSoftObjectPath SoftPath( SubObjectsArray[Index]->GetOuter() );
 
 			FString SoftPathString = SoftPath.ToString();
 			MemAr << SoftPathString;
@@ -225,6 +227,20 @@ namespace DataprepSnapshotUtil
 
 	void ReadSnapshotData(UObject* Object, const TArray<uint8>& InSerializedData, TMap<FString, UClass*>& InClassesMap, TArray<UObject*>& ObjectsToDelete)
 	{
+		// Remove all objects created by default that InObject is dependent on
+		// This method must obviously be called just after the InObject is created
+		auto RemoveDefaultDependencies = [&ObjectsToDelete](UObject* InObject)
+		{
+			TArray< UObject* > ObjectsWithOuter;
+			GetObjectsWithOuter( InObject, ObjectsWithOuter, /*bIncludeNestedObjects = */ true );
+
+			for(UObject* ObjectWithOuter : ObjectsWithOuter)
+			{
+				ObjectWithOuter->Rename( nullptr, GetTransientPackage(), REN_DontCreateRedirectors | REN_NonTransactional );
+				ObjectsToDelete.Add( ObjectWithOuter );
+			}
+		};
+
 		TArray<uint8> MemoryBuffer;
 		if(bUseCompression)
 		{
@@ -241,6 +257,8 @@ namespace DataprepSnapshotUtil
 			checkf(bSucceeded, TEXT("zlib failed to uncompress, which is very unexpected"));
 		}
 
+		RemoveDefaultDependencies( Object );
+
 		FMemoryReader MemAr(bUseCompression ? MemoryBuffer : InSerializedData);
 		FObjectAndNameAsStringProxyArchive Ar(MemAr, false);
 		Ar.SetIsTransacting(true);
@@ -263,19 +281,24 @@ namespace DataprepSnapshotUtil
 
 			UObject* SubObject = NewObject<UObject>( Object, *SubObjectClassPtr, NAME_None, RF_Transient );
 			SubObjectsArray.Add( SubObject );
+
+			RemoveDefaultDependencies( SubObject );
 		}
 
 		// Restore sub-objects' outer if original outer differs from Object
-		for(UObject* SubObject : SubObjectsArray)
+		// Restoration is done in the order the serialization was done: reverse order
+		for(int32 Index = SubObjectsArray.Num() - 1; Index >= 0; --Index)
 		{
 			FString SoftPathString;
 			MemAr << SoftPathString;
 
-			FSoftObjectPath SoftPath( SoftPathString );
+			const FSoftObjectPath SoftPath( SoftPathString );
+
 			UObject* NewOuter = SoftPath.ResolveObject();
 			ensure( NewOuter );
 
-			if( NewOuter && NewOuter != SubObject->GetOuter() )
+			UObject* SubObject = SubObjectsArray[Index];
+			if( NewOuter != SubObject->GetOuter() )
 			{
 				SubObject->Rename( nullptr, NewOuter, REN_DontCreateRedirectors | REN_NonTransactional );
 			}
@@ -285,27 +308,6 @@ namespace DataprepSnapshotUtil
 		for(UObject* SubObject : SubObjectsArray)
 		{
 			SubObject->Serialize(Ar);
-
-			// Some sub-objects might have been created by default when their owner is created.
-			// Duplicates must be deleted before loading the owner
-			TArray< UObject* > SiblingObjectsArray;
-			GetObjectsWithOuter( SubObject->GetOuter(), SiblingObjectsArray, /*bIncludeNestedObjects = */ true );
-			if(SiblingObjectsArray.Num() > 1)
-			{
-				ObjectsToDelete.Reserve( ObjectsToDelete.Num() + SiblingObjectsArray.Num() );
-
-				const FName SubObjectName = SubObject->GetFName();
-				const UClass* SubObjectClass = SubObject->GetClass();
-
-				for(UObject* SiblingObject : SiblingObjectsArray)
-				{
-					if(SiblingObject != SubObject && SiblingObject->GetFName() == SubObjectName && SiblingObject->GetClass() == SubObjectClass)
-					{
-						SiblingObject->Rename( nullptr, GetTransientPackage(), REN_DontCreateRedirectors | REN_NonTransactional );
-						ObjectsToDelete.Add( SiblingObject );
-					}
-				}
-			}
 		}
 
 		// Deserialize object
