Fixed issue with temporarily created directories and files remains on hard drive

#rnx

diff --git a/Engine/Plugins/Experimental/DataPrepEditor/Source/DataPrepEditor/Private/DataPrepEditor.cpp b/Engine/Plugins/Experimental/DataPrepEditor/Source/DataPrepEditor/Private/DataPrepEditor.cpp
--- a/Engine/Plugins/Experimental/DataPrepEditor/Source/DataPrepEditor/Private/DataPrepEditor.cpp
+++ b/Engine/Plugins/Experimental/DataPrepEditor/Source/DataPrepEditor/Private/DataPrepEditor.cpp
@@ -213,9 +213,54 @@ FDataprepEditor::~FDataprepEditor()
 
 	FCoreUObjectDelegates::OnObjectPropertyChanged.RemoveAll(this);
 
-	if (IFileManager::Get().DirectoryExists( *TempDir ))
+	auto DeleteDirectory = [&](const FString& DirectoryToDelete)
 	{
-		IFileManager::Get().DeleteDirectory( *TempDir, true, true );
+		const FString AbsolutePath = FPaths::ConvertRelativePathToFull( DirectoryToDelete );
+		IFileManager::Get().DeleteDirectory( *AbsolutePath, false, true );
+	};
+
+	// Clean up temporary directories and data created for this session
+	{
+
+		DeleteDirectory(TempDir );
+
+		FString PackagePathToDeleteOnDisk;
+		if (FPackageName::TryConvertLongPackageNameToFilename( GetTransientContentFolder(), PackagePathToDeleteOnDisk))
+		{
+			DeleteDirectory( PackagePathToDeleteOnDisk );
+		}
+
+	}
+
+	// Clean up temporary directories associated to process if no session of Dataprep editor is running
+	{
+		auto IsDirectoryEmpty = [&](const TCHAR* Directory) -> bool
+		{
+			bool bDirectoryIsEmpty = true;
+			IFileManager::Get().IterateDirectory(Directory, [&](const TCHAR* FilenameOrDirectory, bool bIsDirectory) -> bool
+			{
+				bDirectoryIsEmpty = false;
+				return false;
+			});
+
+			return bDirectoryIsEmpty;
+		};
+
+		FString RootTempDir = FPaths::Combine( GetRootTemporaryDir(), FString::FromInt( FPlatformProcess::GetCurrentProcessId() ) );
+		if(IsDirectoryEmpty( *RootTempDir ))
+		{
+			DeleteDirectory( RootTempDir );
+		}
+
+		const FString PackagePathToDelete = FPaths::Combine( GetRootPackagePath(), FString::FromInt( FPlatformProcess::GetCurrentProcessId() ) );
+		FString PackagePathToDeleteOnDisk;
+		if (FPackageName::TryConvertLongPackageNameToFilename(PackagePathToDelete, PackagePathToDeleteOnDisk))
+		{
+			if(IsDirectoryEmpty( *PackagePathToDeleteOnDisk ))
+			{
+				DeleteDirectory( PackagePathToDeleteOnDisk );
+			}
+		}
 	}
 }
 
@@ -292,6 +337,70 @@ void FDataprepEditor::UnregisterTabSpawners(const TSharedRef<class FTabManager>&
 	// end of temp code for nodes development
 }
 
+void FDataprepEditor::CleanUpTemporaryDirectories()
+{
+	const int32 CurrentProcessID = FPlatformProcess::GetCurrentProcessId();
+
+	TSet<FString> TempDirectories;
+	IFileManager::Get().IterateDirectory( *GetRootTemporaryDir(), [&](const TCHAR* FilenameOrDirectory, bool bIsDirectory) -> bool
+	{
+		if (bIsDirectory)
+		{
+			FString DirectoryName = FPaths::GetBaseFilename( FilenameOrDirectory );
+			if(FCString::IsNumeric( *DirectoryName ))
+			{
+
+				int32 ProcessID = FCString::Atoi( *DirectoryName );
+				if(ProcessID != CurrentProcessID)
+				{
+					FProcHandle ProcHandle = FPlatformProcess::OpenProcess( ProcessID );
+
+					// Delete directories if process is not valid
+					bool bDeleteDirectories = !ProcHandle.IsValid();
+
+					// Process is valid, check if application associated with process id is UE4 editor
+					if(!bDeleteDirectories)
+					{
+						const FString ApplicationName = FPlatformProcess::GetApplicationName( ProcessID );
+						bDeleteDirectories = !ApplicationName.StartsWith(TEXT("UE4Editor"));
+					}
+
+					if(bDeleteDirectories)
+					{
+						FString PackagePathToDelete = FPaths::Combine( GetRootPackagePath(), DirectoryName );
+						FString PackagePathToDeleteOnDisk;
+						if (FPackageName::TryConvertLongPackageNameToFilename(PackagePathToDelete, PackagePathToDeleteOnDisk))
+						{
+							TempDirectories.Add( MoveTemp(PackagePathToDeleteOnDisk) );
+						}
+
+						TempDirectories.Emplace( FilenameOrDirectory );
+					}
+				}
+			}
+		}
+		return true;
+	});
+
+	for(FString& TempDirectory : TempDirectories)
+	{
+		FString AbsolutePath = FPaths::ConvertRelativePathToFull( TempDirectory );
+		IFileManager::Get().DeleteDirectory( *AbsolutePath, false, true );
+	}
+}
+
+const FString& FDataprepEditor::GetRootTemporaryDir()
+{
+	static FString RootTemporaryDir = FPaths::Combine(FPaths::ProjectIntermediateDir(), TEXT("DataprepTemp") );
+	return RootTemporaryDir;
+}
+
+const FString& FDataprepEditor::GetRootPackagePath()
+{
+	static FString RootPackagePath( TEXT("/DataprepEditor/Transient") );
+	return RootPackagePath;
+}
+
 void FDataprepEditor::InitDataprepEditor(const EToolkitMode::Type Mode, const TSharedPtr<IToolkitHost>& InitToolkitHost, UDataprepAsset* InDataprepAsset)
 {
 	DataprepAssetPtr = TWeakObjectPtr<UDataprepAsset>(InDataprepAsset);
@@ -304,8 +413,9 @@ void FDataprepEditor::InitDataprepEditor(const EToolkitMode::Type Mode, const TS
 	// Assign unique session identifier
 	SessionID = FGuid::NewGuid().ToString();
 
-	// Create temporary directory which will be used by UDatasmithStaticMeshCADImportData to store transient data
-	TempDir = FPaths::Combine(FPaths::ProjectIntermediateDir(), TEXT("DataprepTemp"), SessionID);
+	// Create temporary directory to store transient data
+	CleanUpTemporaryDirectories();
+	TempDir = FPaths::Combine( GetRootTemporaryDir(), FString::FromInt( FPlatformProcess::GetCurrentProcessId() ), SessionID);
 	IFileManager::Get().MakeDirectory(*TempDir);
 
 	// Temp code for the nodes development
@@ -1032,7 +1142,7 @@ bool FDataprepEditor::CanCommitWorld()
 
 FString FDataprepEditor::GetTransientContentFolder()
 {
-	return FPaths::Combine( TEXT("/DataPrepEditor"), SessionID );
+	return FPaths::Combine( GetRootPackagePath(), FString::FromInt( FPlatformProcess::GetCurrentProcessId() ), SessionID );
 }
 
 void DataprepEditorUtil::DeleteTemporaryPackage( const FString& PathToDelete )
diff --git a/Engine/Plugins/Experimental/DataPrepEditor/Source/DataPrepEditor/Private/DataPrepEditor.h b/Engine/Plugins/Experimental/DataPrepEditor/Source/DataPrepEditor/Private/DataPrepEditor.h
--- a/Engine/Plugins/Experimental/DataPrepEditor/Source/DataPrepEditor/Private/DataPrepEditor.h
+++ b/Engine/Plugins/Experimental/DataPrepEditor/Source/DataPrepEditor/Private/DataPrepEditor.h
@@ -86,6 +86,12 @@ public:
 	/** Gets or sets the flag for context sensitivity in the graph action menu */
 	bool& GetIsContextSensitive() { return bIsActionMenuContextSensitive; }
 
+	/** Returns root package which all transient packages are created under */
+	static const FString& GetRootPackagePath();
+
+	/** Returns root directory which all transient directories and data are created under */
+	static const FString& GetRootTemporaryDir();
+
 private:
 	void BindCommands();
 	void OnSaveScene();
@@ -148,9 +154,12 @@ private:
 	/** Handles changes in the Dataprep asset */
 	void OnDataprepAssetChanged(FDataprepAssetChangeType ChangeType, int32 Index);
 
-	/** Handles change to the dataprep pipeline */
+	/** Handles change to the Dataprep pipeline */
 	void OnDataprepPipelineChange(UObject* ChangedObject);
 
+	/** Remove all temporary data remaining from previous runs of the Dataprep editor */
+	void CleanUpTemporaryDirectories();
+
 private:
 	bool bWorldBuilt;
 	bool bIsFirstRun;
diff --git a/Engine/Plugins/Experimental/DataPrepEditor/Source/DataPrepEditor/Private/DataPrepEditorModule.cpp b/Engine/Plugins/Experimental/DataPrepEditor/Source/DataPrepEditor/Private/DataPrepEditorModule.cpp
--- a/Engine/Plugins/Experimental/DataPrepEditor/Source/DataPrepEditor/Private/DataPrepEditorModule.cpp
+++ b/Engine/Plugins/Experimental/DataPrepEditor/Source/DataPrepEditor/Private/DataPrepEditorModule.cpp
@@ -9,6 +9,7 @@
 
 #include "Developer/AssetTools/Public/IAssetTools.h"
 #include "Developer/AssetTools/Public/AssetToolsModule.h"
+#include "Misc/PackageName.h"
 #include "Modules/ModuleManager.h"
 #include "PropertyEditorModule.h"
 #include "UObject/StrongObjectPtr.h"
@@ -43,6 +44,9 @@ public:
 		TSharedPtr<FAssetTypeActions_DataprepAsset> DataprepAssetTypeAction = MakeShareable(new FAssetTypeActions_DataprepAsset);
 		AssetTools.RegisterAssetTypeActions(DataprepAssetTypeAction.ToSharedRef());
 		AssetTypeActionsArray.Add(DataprepAssetTypeAction);
+
+		// Register mount point for Dataprep editors root package folder
+		FPackageName::RegisterMountPoint( FDataprepEditor::GetRootPackagePath() + TEXT("/"), FDataprepEditor::GetRootTemporaryDir() );
 	}
 
 	virtual void ShutdownModule() override
@@ -62,6 +66,9 @@ public:
 		AssetTypeActionsArray.Empty();
 
 		FDataprepEditorStyle::Shutdown();
+
+		// Unregister mount point for Dataprep editors root package folder
+		FPackageName::UnRegisterMountPoint( FDataprepEditor::GetRootPackagePath() + TEXT("/"), FDataprepEditor::GetRootTemporaryDir() );
 	}
 
 	/** Gets the extensibility managers for outside entities to extend datasmith data prep editor's menus and toolbars */
