Fixed localization gather warnings

#rnx

diff --git a/Engine/Plugins/Animation/LiveLink/Source/LiveLinkEditor/Private/SLiveLinkDataView.cpp b/Engine/Plugins/Animation/LiveLink/Source/LiveLinkEditor/Private/SLiveLinkDataView.cpp
--- a/Engine/Plugins/Animation/LiveLink/Source/LiveLinkEditor/Private/SLiveLinkDataView.cpp
+++ b/Engine/Plugins/Animation/LiveLink/Source/LiveLinkEditor/Private/SLiveLinkDataView.cpp
@@ -81,7 +81,7 @@ void SLiveLinkDataView::Construct(const FArguments& Args, FLiveLinkClient* InCli
 			.MaxValue(10.0)
 			.MinSliderValue(0.0)
 			.MaxSliderValue(10.0)
-			.ToolTipText(LOCTEXT("RefreshDelay", "Refresh delay of Static & Frame data."))
+			.ToolTipText(LOCTEXT("RefreshDelay_ToolTip", "Refresh delay of Static & Frame data."))
 			.Value(this, &SLiveLinkDataView::GetRefreshDelay)
 			.OnValueCommitted(this, &SLiveLinkDataView::SetRefreshDelayInternal)
 			.IsEnabled(this, &SLiveLinkDataView::CanEditRefreshDelay);
diff --git a/Engine/Plugins/Developer/Concert/MultiUserClient/Source/MultiUserClient/Private/Widgets/SConcertBrowser.cpp b/Engine/Plugins/Developer/Concert/MultiUserClient/Source/MultiUserClient/Private/Widgets/SConcertBrowser.cpp
--- a/Engine/Plugins/Developer/Concert/MultiUserClient/Source/MultiUserClient/Private/Widgets/SConcertBrowser.cpp
+++ b/Engine/Plugins/Developer/Concert/MultiUserClient/Source/MultiUserClient/Private/Widgets/SConcertBrowser.cpp
@@ -2727,7 +2727,7 @@ TSharedRef<SWidget> SConcertSessionBrowser::MakeSessionViewOptionsBar()
 		FMenuBuilder MenuBuilder(/*bInShouldCloseWindowAfterMenuSelection=*/true, nullptr);
 
 		MenuBuilder.AddMenuEntry(
-			LOCTEXT("ActiveSessions", "Active Sessions"),
+			LOCTEXT("ActiveSessions_Label", "Active Sessions"),
 			LOCTEXT("ActiveSessions_Tooltip", "Displays Active Sessions"),
 			FSlateIcon(),
 			FUIAction(
@@ -2739,7 +2739,7 @@ TSharedRef<SWidget> SConcertSessionBrowser::MakeSessionViewOptionsBar()
 		);
 
 		MenuBuilder.AddMenuEntry(
-			LOCTEXT("ArchivedSessions", "Archived Sessions"),
+			LOCTEXT("ArchivedSessions_Label", "Archived Sessions"),
 			LOCTEXT("ArchivedSessions_Tooltip", "Displays Archived Sessions"),
 			FSlateIcon(),
 			FUIAction(
@@ -2751,7 +2751,7 @@ TSharedRef<SWidget> SConcertSessionBrowser::MakeSessionViewOptionsBar()
 		);
 
 		MenuBuilder.AddMenuEntry(
-			LOCTEXT("DefaultServer", "Default Server Sessions"),
+			LOCTEXT("DefaultServer_Label", "Default Server Sessions"),
 			LOCTEXT("DefaultServer_Tooltip", "Displays Sessions Hosted By the Default Server"),
 			FSlateIcon(),
 			FUIAction(
diff --git a/Engine/Plugins/Experimental/DataPrepEditor/Source/DataPrepCore/Private/DataPrepAsset.cpp b/Engine/Plugins/Experimental/DataPrepEditor/Source/DataPrepCore/Private/DataPrepAsset.cpp
--- a/Engine/Plugins/Experimental/DataPrepEditor/Source/DataPrepCore/Private/DataPrepAsset.cpp
+++ b/Engine/Plugins/Experimental/DataPrepEditor/Source/DataPrepCore/Private/DataPrepAsset.cpp
@@ -123,13 +123,13 @@ void UDataprepAsset::RunProducers(const UDataprepContentProducer::ProducerContex
 
 	OutAssets.Empty();
 
-	FDataprepProgressTask Task( *InContext.ProgressReporterPtr, NSLOCTEXT( "DataprepAsset", "RunProducers", "Importing ..." ), (float)Producers.Num(), 1.0f );
+	FDataprepProgressTask Task( *InContext.ProgressReporterPtr, NSLOCTEXT( "DataprepAsset", "ProducerReport_Importing", "Importing ..." ), (float)Producers.Num(), 1.0f );
 
 	for ( FDataprepAssetProducer& AssetProducer : Producers )
 	{
 		if( UDataprepContentProducer* Producer = AssetProducer.Producer )
 		{
-			Task.ReportNextStep( FText::Format( NSLOCTEXT( "DataprepAsset", "ProducerReport", "Importing {0} ..."), FText::FromString( Producer->GetName() ) ) );
+			Task.ReportNextStep( FText::Format( NSLOCTEXT( "DataprepAsset", "ProducerReport_ImportingX", "Importing {0} ..."), FText::FromString( Producer->GetName() ) ) );
 
 			// Run producer if enabled and, if superseded, superseder is disabled
 			const bool bIsOkToRun = AssetProducer.bIsEnabled &&	( AssetProducer.SupersededBy == INDEX_NONE || !Producers[AssetProducer.SupersededBy].bIsEnabled );
@@ -149,7 +149,7 @@ void UDataprepAsset::RunProducers(const UDataprepContentProducer::ProducerContex
 					}
 					else
 					{
-						OutReason = FText::Format(NSLOCTEXT("DataprepAsset", "ProducerRunFailed", "{0} failed to run."), FText::FromString( Producer->GetName() ) ).ToString();
+						OutReason = FText::Format(NSLOCTEXT("DataprepAsset", "ProducerReport_Failed", "{0} failed to run."), FText::FromString( Producer->GetName() ) ).ToString();
 					}
 				}
 
@@ -163,7 +163,7 @@ void UDataprepAsset::RunProducers(const UDataprepContentProducer::ProducerContex
 		}
 		else
 		{
-			Task.ReportNextStep( NSLOCTEXT( "DataprepAsset", "ProducerReport", "Skipped invalid producer ...") );
+			Task.ReportNextStep( NSLOCTEXT( "DataprepAsset", "ProducerReport_Skipped", "Skipped invalid producer ...") );
 		}
 	}
 }
diff --git a/Engine/Plugins/Experimental/DataPrepEditor/Source/DataPrepEditor/Private/Widgets/DataprepWidgets.cpp b/Engine/Plugins/Experimental/DataPrepEditor/Source/DataPrepEditor/Private/Widgets/DataprepWidgets.cpp
--- a/Engine/Plugins/Experimental/DataPrepEditor/Source/DataPrepEditor/Private/Widgets/DataprepWidgets.cpp
+++ b/Engine/Plugins/Experimental/DataPrepEditor/Source/DataPrepEditor/Private/Widgets/DataprepWidgets.cpp
@@ -195,7 +195,7 @@ TSharedRef<SWidget> SDataprepConsumerWidget::BuildNullWidget()
 			[
 				SNew(STextBlock)
 				.Font( IDetailLayoutBuilder::GetDetailFont() )
-				.Text(LOCTEXT("DataprepSlateHelper_ContentFolderLabel", "Error: Not a valid consumer"))
+				.Text(LOCTEXT("DataprepSlateHelper_Error_InvalidConsumer", "Error: Not a valid consumer"))
 				.Margin(FMargin( 5.0f, 5.0f, 0.0f, 0.0f ) )
 				.ColorAndOpacity( FLinearColor(1,0,0,1) )
 			]
diff --git a/Engine/Plugins/Experimental/PythonScriptPlugin/Source/PythonScriptPlugin/Private/PythonScriptPlugin.cpp b/Engine/Plugins/Experimental/PythonScriptPlugin/Source/PythonScriptPlugin/Private/PythonScriptPlugin.cpp
--- a/Engine/Plugins/Experimental/PythonScriptPlugin/Source/PythonScriptPlugin/Private/PythonScriptPlugin.cpp
+++ b/Engine/Plugins/Experimental/PythonScriptPlugin/Source/PythonScriptPlugin/Private/PythonScriptPlugin.cpp
@@ -373,7 +373,7 @@ private:
 		{
 			bool bOpened = DesktopPlatform->OpenFileDialog(
 				FSlateApplication::Get().FindBestParentWindowHandleForDialogs(nullptr),
-				LOCTEXT("OpenPython", "Execute Python Script").ToString(),
+				LOCTEXT("ExecutePython", "Execute Python Script").ToString(),
 				DefaultDirectory,
 				TEXT(""),
 				TEXT("Python files|*.py|"),
diff --git a/Engine/Plugins/Media/HAPMedia/Source/HAPMedia/Private/HAPMediaModule.cpp b/Engine/Plugins/Media/HAPMedia/Source/HAPMedia/Private/HAPMediaModule.cpp
--- a/Engine/Plugins/Media/HAPMedia/Source/HAPMedia/Private/HAPMediaModule.cpp
+++ b/Engine/Plugins/Media/HAPMedia/Source/HAPMedia/Private/HAPMediaModule.cpp
@@ -66,5 +66,6 @@ private:
 	/** Whether the module has been initialized. */
 };
 
+#undef LOCTEXT_NAMESPACE
 
 IMPLEMENT_MODULE(FHAPMediaModule, HAPMedia);
