Fixed crash in UK2Node_DataprepProducer

diff --git a/Engine/Plugins/Experimental/DataPrepEditor/Source/DataPrepEditor/Private/BlueprintNodes/K2Node_DataprepProducer.cpp b/Engine/Plugins/Experimental/DataPrepEditor/Source/DataPrepEditor/Private/BlueprintNodes/K2Node_DataprepProducer.cpp
--- a/Engine/Plugins/Experimental/DataPrepEditor/Source/DataPrepEditor/Private/BlueprintNodes/K2Node_DataprepProducer.cpp
+++ b/Engine/Plugins/Experimental/DataPrepEditor/Source/DataPrepEditor/Private/BlueprintNodes/K2Node_DataprepProducer.cpp
@@ -29,7 +29,7 @@
 const char* DataprepProducerFontName = "FontAwesome.11";
 
 // #ueent_todo: Temporarily disabling details view in producer K2Node
-static bool bShowVisualWidget = false;
+//#define WIDGET_HAS_UI
 
 void UK2Node_DataprepProducer::AllocateDefaultPins()
 {
@@ -72,25 +72,33 @@ bool UK2Node_DataprepProducer::CanUserDeleteNode() const
 
 void UK2Node_DataprepProducer::SetDataprepAsset( UDataprepAsset* InDataprepAsset )
 {
-	// #ueent_todo: Should we allow changing the Dataprep asset associated with a UK2Node_DataprepProducer?
-	check(DataprepAsset == nullptr);
+#ifdef WIDGET_HAS_UI
+	check(InDataprepAsset);
 	DataprepAsset = InDataprepAsset;
 	DataprepAssetPath = FSoftObjectPath( InDataprepAsset );
+#endif
 }
 
 UDataprepAsset* UK2Node_DataprepProducer::GetDataprepAsset()
 {
+#ifdef WIDGET_HAS_UI
+	return nullptr;
+#else
 	return DataprepAsset;
+#endif
 }
 
 void UK2Node_DataprepProducer::Serialize(FArchive & Ar)
 {
 	Super::Serialize( Ar );
 
+#ifdef WIDGET_HAS_UI
 	if(Ar.IsLoading())
 	{
+		// #ueent_todo: What if the asset has been moved by the user without redirection?
 		DataprepAsset = Cast<UDataprepAsset>(DataprepAssetPath.TryLoad());
 	}
+#endif
 }
 
 #if WITH_EDITOR
@@ -113,6 +121,11 @@ TSharedPtr<SGraphNode> UK2Node_DataprepProducer::CreateVisualWidget()
 		SLATE_BEGIN_ARGS(SGraphNodeDataprepProducer){}
 		SLATE_END_ARGS()
 
+		SGraphNodeDataprepProducer()
+			: DataprepAsset(nullptr)
+		{
+		}
+
 		~SGraphNodeDataprepProducer()
 		{
 			if( DataprepAsset )
@@ -152,9 +165,11 @@ TSharedPtr<SGraphNode> UK2Node_DataprepProducer::CreateVisualWidget()
 		{
 			bNeedsUpdate = false;
 
+#ifdef WIDGET_HAS_UI
 			DataprepAsset = CastChecked<UK2Node_DataprepProducer>(GraphNode)->GetDataprepAsset();
+			check(DataprepAsset);
 
-			DataprepAsset->GetOnChanged().AddRaw( this, &SGraphNodeDataprepProducer::OnDataprepAssetChanged );
+			DataprepAsset->GetOnChanged().AddSP( this, &SGraphNodeDataprepProducer::OnDataprepAssetChanged );
 
 			// #ueent_todo: Move this code where it would capture new classes of producer created at runtime, .i.e BP based producers
 			for( TObjectIterator< UClass > It ; It ; ++It )
@@ -167,10 +182,12 @@ TSharedPtr<SGraphNode> UK2Node_DataprepProducer::CreateVisualWidget()
 					ProducerDescriptions.Emplace( CurrentClass, Producer->GetLabel(), Producer->GetDescription() );
 				}
 			}
+#endif
 		}
 
 		void OnAddProducer( int32 Index )
 		{
+#ifdef WIDGET_HAS_UI
 			if (DataprepAsset == nullptr || !ProducerDescriptions.IsValidIndex( Index ) )
 			{
 				return;
@@ -180,10 +197,12 @@ TSharedPtr<SGraphNode> UK2Node_DataprepProducer::CreateVisualWidget()
 			check( ProducerClass );
 
 			DataprepAsset->AddProducer( ProducerClass );
+#endif
 		}
 
 		TSharedRef<SWidget> CreateAddProducerMenuWidget()
 		{
+#ifdef WIDGET_HAS_UI
 			FMenuBuilder MenuBuilder(true, nullptr);
 
 			MenuBuilder.BeginSection("AddNewProducer", LOCTEXT("DataprepEditorViews_AddImports", "Add Producer"));
@@ -210,10 +229,14 @@ TSharedPtr<SGraphNode> UK2Node_DataprepProducer::CreateVisualWidget()
 			MenuBuilder.EndSection();
 
 			return MenuBuilder.MakeWidget();
+#else
+			return SNullWidget::NullWidget;
+#endif
 		}
 
 		TSharedRef< SWidget > CreateProducerWidget( ProducerWidget* ProducerWidgetPtr, int32 Index )
 		{
+#ifdef WIDGET_HAS_UI
 			auto CheckEntry = [this, Index]()
 			{
 				FScopedTransaction Transaction(LOCTEXT("DataprepAsset_AddProducer", "AddProducer"));
@@ -295,15 +318,14 @@ TSharedPtr<SGraphNode> UK2Node_DataprepProducer::CreateVisualWidget()
 			DetailsView->SetObject( DataprepAsset->GetProducer( Index ) );
 
 			return Widget.ToSharedRef();
+#else
+			return SNullWidget::NullWidget;
+#endif
 		}
 
 		virtual void CreateBelowPinControls(TSharedPtr<SVerticalBox> MainBox) override
 		{
-			if(!bShowVisualWidget)
-			{
-				return;
-			}
-
+#ifdef WIDGET_HAS_UI
 			TSharedPtr<SWidget> AddNewMenu = SNew(SComboButton)
 			.ComboButtonStyle(FEditorStyle::Get(), "ToolbarComboButton")
 			.ButtonStyle(FEditorStyle::Get(), "FlatButton.Success")
@@ -354,6 +376,7 @@ TSharedPtr<SGraphNode> UK2Node_DataprepProducer::CreateVisualWidget()
 					CreateProducerWidget( &ProducerWidgets[ Index ], Index )
 				];
 			}
+#endif
 		}
 
 	private:
