Merging CL #9421376 from //UE4/Dev-Enterprise

Fix flaw in logic about tracking changes during the execution of the pipeline

#rnx

diff --git a/Engine/Plugins/Experimental/DataPrepEditor/Source/DataPrepCore/Private/DataprepActionAsset.cpp b/Engine/Plugins/Experimental/DataPrepEditor/Source/DataPrepCore/Private/DataprepActionAsset.cpp
--- a/Engine/Plugins/Experimental/DataPrepEditor/Source/DataPrepCore/Private/DataprepActionAsset.cpp
+++ b/Engine/Plugins/Experimental/DataPrepEditor/Source/DataPrepCore/Private/DataprepActionAsset.cpp
@@ -322,6 +322,10 @@ void UDataprepActionAsset::ExecuteAction(const TSharedPtr<FDataprepActionContext
 	{
 		if ( UDataprepOperation* Operation = Step->Operation )
 		{
+			// Cache number of assets and objects before execution
+			int32 AssetsDiffCount = ContextPtr->Assets.Num();
+			int32 ActorsDiffCount = OperationContext->Context->Objects.Num();
+
 			TSharedRef<FDataprepOperationContext> OperationContextPtr = this->OperationContext.ToSharedRef();
 			Operation->ExecuteOperation( OperationContextPtr );
 
@@ -332,7 +336,9 @@ void UDataprepActionAsset::ExecuteAction(const TSharedPtr<FDataprepActionContext
 			}
 
 			// Process the changes in the context if applicable
-			this->ProcessWorkingSetChanged();
+			AssetsDiffCount = ContextPtr->Assets.Num() - AssetsDiffCount;
+			ActorsDiffCount = OperationContext->Context->Objects.Num() - ActorsDiffCount - AssetsDiffCount;
+			this->ProcessWorkingSetChanged( AssetsDiffCount != 0, ActorsDiffCount != 0 );
 		}
 		else if ( UDataprepFilter* Filter = Step->Filter )
 		{
@@ -462,12 +468,12 @@ void UDataprepActionAsset::OnDeleteObjects(TArray<UObject*> Objects)
 	}
 }
 
-void UDataprepActionAsset::ProcessWorkingSetChanged()
+void UDataprepActionAsset::ProcessWorkingSetChanged( bool bAddedAssets, bool bAddedActors )
 {
 	if(bWorkingSetHasChanged && ContextPtr != nullptr)
 	{
-		bool bAssetsChanged = false;
-		bool bWorldChanged = false;
+		bool bAssetsChanged = bAddedAssets;
+		bool bWorldChanged = bAddedActors;
 
 		TSet<UObject*> SelectedObjectSet( OperationContext->Context->Objects );
 
diff --git a/Engine/Plugins/Experimental/DataPrepEditor/Source/DataPrepCore/Public/DataprepActionAsset.h b/Engine/Plugins/Experimental/DataPrepEditor/Source/DataPrepCore/Public/DataprepActionAsset.h
--- a/Engine/Plugins/Experimental/DataPrepEditor/Source/DataPrepCore/Public/DataprepActionAsset.h
+++ b/Engine/Plugins/Experimental/DataPrepEditor/Source/DataPrepCore/Public/DataprepActionAsset.h
@@ -293,8 +293,12 @@ private:
 	 */
 	void OnDeleteObjects( TArray<UObject*> Objects);
 
-	/** Executes the deletion and removal requested by an operation after its execution */
-	void ProcessWorkingSetChanged();
+	/**
+	 * Executes the deletion and removal requested by an operation after its execution
+	 * @param bAddedAssets		Indicates that assets have been added during the operation's execution
+	 * @param bAddedActors		Indicates that actors have been added during the operation's execution
+	 */
+	void ProcessWorkingSetChanged(bool bAddedAssets, bool bAddedActors);
 
 	/** Returns the outer to be used according to an asset's class */
 	UObject* GetAssetOuterByClass( UClass* AssetClass );
