Merging CL #9421544 from //UE4/Dev-Enterprise

Fixed package path for transient objects created in DataprepCore and DataprepEditor: Prefixed path with /Engine

diff --git a/Engine/Plugins/Experimental/DataPrepEditor/Source/DataPrepCore/Private/DataprepCorePrivateUtils.cpp b/Engine/Plugins/Experimental/DataPrepEditor/Source/DataPrepCore/Private/DataprepCorePrivateUtils.cpp
--- a/Engine/Plugins/Experimental/DataPrepEditor/Source/DataPrepCore/Private/DataprepCorePrivateUtils.cpp
+++ b/Engine/Plugins/Experimental/DataPrepEditor/Source/DataPrepCore/Private/DataprepCorePrivateUtils.cpp
@@ -72,7 +72,7 @@ const FString& DataprepCorePrivateUtils::GetRootTemporaryDir()
 
 const FString& DataprepCorePrivateUtils::GetRootPackagePath()
 {
-	static FString RootPackagePath( TEXT("/DataprepCore/Transient") );
+	static FString RootPackagePath( TEXT("/Engine/DataprepCore/Transient") );
 	return RootPackagePath;
 }
 
diff --git a/Engine/Plugins/Experimental/DataPrepEditor/Source/DataPrepEditor/Private/DataPrepEditor.cpp b/Engine/Plugins/Experimental/DataPrepEditor/Source/DataPrepEditor/Private/DataPrepEditor.cpp
--- a/Engine/Plugins/Experimental/DataPrepEditor/Source/DataPrepEditor/Private/DataPrepEditor.cpp
+++ b/Engine/Plugins/Experimental/DataPrepEditor/Source/DataPrepEditor/Private/DataPrepEditor.cpp
@@ -508,7 +508,7 @@ const FString& FDataprepEditor::GetRootTemporaryDir()
 
 const FString& FDataprepEditor::GetRootPackagePath()
 {
-	static FString RootPackagePath( TEXT("/DataprepEditor/Transient") );
+	static FString RootPackagePath( TEXT("/Engine/DataprepEditor/Transient") );
 	return RootPackagePath;
 }
 
