Merging  CLs #9421605 and #9422888 from //UE4/Dev-Enterprise

9421605:
Fixed crash on import: Name of scene actor could collide with unique ID of one of its related actors. Must check name before assigning.

9422888:
Improvement to the dataprep parameterisation.

We can now remove binded properties.
Transaction are now also supported.

diff --git a/Engine/Plugins/Enterprise/DatasmithImporter/Source/DatasmithImporter/Private/DatasmithFileProducer.cpp b/Engine/Plugins/Enterprise/DatasmithImporter/Source/DatasmithImporter/Private/DatasmithFileProducer.cpp
--- a/Engine/Plugins/Enterprise/DatasmithImporter/Source/DatasmithImporter/Private/DatasmithFileProducer.cpp
+++ b/Engine/Plugins/Enterprise/DatasmithImporter/Source/DatasmithImporter/Private/DatasmithFileProducer.cpp
@@ -463,19 +463,27 @@ void UDatasmithFileProducer::PreventNameCollision()
 						{
 							if( UDatasmithAssetUserData* AssetUserData = UDatasmithContentBlueprintLibrary::GetDatasmithUserData( Actor ) )
 							{
-								if( FString* Value = AssetUserData->MetaData.Find( UDatasmithAssetUserData::UniqueIdMetaDataKey ) )
+								if( FString* ValuePtr = AssetUserData->MetaData.Find( UDatasmithAssetUserData::UniqueIdMetaDataKey ) )
 								{
 									FSoftObjectPath PreviousActorSoftPath(Actor);
 
+									FString& Value = *ValuePtr;
+
 									// Set Actor's name to the one from its old unique Id.
 									// Rationale: The unique Id is used to reconstruct the IDatasmithActorElement in the Datasmith consumer.
 									// Important Note: No need to prefix the actor's name with the namespace, it will be done by the parent class, UDataprepContentProducer
-									Actor->Rename( *(*Value) );
+									// Important Note: Value of unique Id might collide with name of scene actor. See JIRA UE-80831
+									if( !Actor->Rename( *Value, nullptr, REN_Test ) )
+									{
+										Value = MakeUniqueObjectName( Actor->GetOuter(), Actor->GetClass(), *Value ).ToString();
+									}
+
+									Actor->Rename( *Value );
 
 									ActorRedirectorMap.Emplace( PreviousActorSoftPath, Actor );
 
 									// Prefix actor's unique Id with the namespace
-									*Value = Namespace + TEXT("_") + *Value;
+									Value = Namespace + TEXT("_") + Value;
 								}
 							}
 						}
diff --git a/Engine/Plugins/Experimental/DataPrepEditor/Source/DataPrepCore/Private/DataPrepAsset.cpp b/Engine/Plugins/Experimental/DataPrepEditor/Source/DataPrepCore/Private/DataPrepAsset.cpp
--- a/Engine/Plugins/Experimental/DataPrepEditor/Source/DataPrepCore/Private/DataPrepAsset.cpp
+++ b/Engine/Plugins/Experimental/DataPrepEditor/Source/DataPrepCore/Private/DataPrepAsset.cpp
@@ -215,6 +215,16 @@ void UDataprepAsset::BindObjectPropertyToParameterization(UObject* Object, const
 
 }
 
+bool UDataprepAsset::IsObjectPropertyBinded(UObject* Object, const TArray<struct FDataprepPropertyLink>& InPropertyChain) const
+{
+	return Parameterization->IsObjectPropertyBinded( Object, InPropertyChain );
+}
+
+void UDataprepAsset::RemoveObjectPropertyFromParameterization(UObject* Object, const TArray<struct FDataprepPropertyLink>& InPropertyChain)
+{
+	Parameterization->RemoveBindedObjectProperty( Object, InPropertyChain );
+}
+
 void UDataprepAsset::OnDataprepBlueprintChanged( UBlueprint* InBlueprint )
 {
 	if(InBlueprint == DataprepRecipeBP)
diff --git a/Engine/Plugins/Experimental/DataPrepEditor/Source/DataPrepCore/Private/DataprepParameterizationArchive.cpp b/Engine/Plugins/Experimental/DataPrepEditor/Source/DataPrepCore/Private/DataprepParameterizationArchive.cpp
--- a/Engine/Plugins/Experimental/DataPrepEditor/Source/DataPrepCore/Private/DataprepParameterizationArchive.cpp
+++ b/Engine/Plugins/Experimental/DataPrepEditor/Source/DataPrepCore/Private/DataprepParameterizationArchive.cpp
@@ -25,7 +25,7 @@ FArchive& FDataprepParameterizationWriter::operator<<(FName& Value)
 FArchive& FDataprepParameterizationWriter::operator<<(UObject*& Value)
 {
 	FString Path;
-	if ( Value )
+	if ( Value && Value->IsValidLowLevelFast() )
 	{
 		if ( UObject* Outer = Value->GetOuter() )
 		{
diff --git a/Engine/Plugins/Experimental/DataPrepEditor/Source/DataPrepCore/Private/DataprepParameterizationArchive.h b/Engine/Plugins/Experimental/DataPrepEditor/Source/DataPrepCore/Private/DataprepParameterizationArchive.h
--- a/Engine/Plugins/Experimental/DataPrepEditor/Source/DataPrepCore/Private/DataprepParameterizationArchive.h
+++ b/Engine/Plugins/Experimental/DataPrepEditor/Source/DataPrepCore/Private/DataprepParameterizationArchive.h
@@ -22,7 +22,7 @@ public:
 
 		ArIgnoreClassRef = true;
 		ArIgnoreArchetypeRef = true;
-		ArNoDelta = !false;
+		ArNoDelta = true;
 		ArPortFlags |= PPF_ForceTaggedSerialization;
 
 #if USE_STABLE_LOCALIZATION_KEYS
@@ -40,21 +40,21 @@ public:
 
 	virtual FArchive& operator<<(FName& Value) override;
 
-	virtual FArchive& operator <<(UObject*& Value) override;
+	virtual FArchive& operator<<(UObject*& Value) override;
 
-	virtual FArchive& operator <<(FText& Value) override
+	virtual FArchive& operator<<(FText& Value) override
 	{
 		// todo
 		return *this;
 	}
 
-	virtual FArchive& operator <<(struct FLazyObjectPtr& Value) override;
+	virtual FArchive& operator<<(struct FLazyObjectPtr& Value) override;
 
-	virtual FArchive& operator <<(struct FSoftObjectPtr& Value) override;
+	virtual FArchive& operator<<(struct FSoftObjectPtr& Value) override;
 
-	virtual FArchive& operator <<(struct FSoftObjectPath& Value) override;
+	virtual FArchive& operator<<(struct FSoftObjectPath& Value) override;
 
-	virtual FArchive& operator <<(struct FWeakObjectPtr& Value) override;
+	virtual FArchive& operator<<(struct FWeakObjectPtr& Value) override;
 
 	virtual FString GetArchiveName() const override;
 
diff --git a/Engine/Plugins/Experimental/DataPrepEditor/Source/DataPrepCore/Private/Parameterization/DataprepParameterization.cpp b/Engine/Plugins/Experimental/DataPrepEditor/Source/DataPrepCore/Private/Parameterization/DataprepParameterization.cpp
--- a/Engine/Plugins/Experimental/DataPrepEditor/Source/DataPrepCore/Private/Parameterization/DataprepParameterization.cpp
+++ b/Engine/Plugins/Experimental/DataPrepEditor/Source/DataPrepCore/Private/Parameterization/DataprepParameterization.cpp
@@ -2,18 +2,20 @@
 
 #include "Parameterization/DataprepParameterization.h"
 
+#include "DataprepAsset.h"
 #include "DataprepParameterizationArchive.h"
 
+#include "CoreGlobals.h"
 #include "Editor.h"
 #include "Engine/Engine.h"
 #include "Math/UnrealMathUtility.h"
 #include "Parameterization/DataprepParameterizationUtils.h"
 #include "Serialization/ObjectReader.h"
+#include "Templates/UnrealTemplate.h"
 #include "UObject/Object.h"
 #include "UObject/PropertyPortFlags.h"
 #include "UObject/UObjectHash.h"
 #include "UObject/UnrealType.h"
-#include "Templates/UnrealTemplate.h"
 
 namespace DataprepParameterization
 {
@@ -314,7 +316,6 @@ namespace DataprepParameterization
 };
 
 
-
 bool FDataprepParameterizationBinding::operator==(const FDataprepParameterizationBinding& Other) const
 {
 	// Not a accurate but ok for now
@@ -339,6 +340,16 @@ uint32 GetTypeHash(const TArray<FDataprepPropertyLink>& PropertyLinks)
 	return Hash;
 }
 
+UDataprepParameterization::UDataprepParameterization()
+{
+	OnObjectModifiedHandle = FCoreUObjectDelegates::OnObjectModified.AddUObject( this, &UDataprepParameterization::OnObjectModified );
+}
+
+UDataprepParameterization::~UDataprepParameterization()
+{
+	FCoreUObjectDelegates::OnObjectModified.Remove( OnObjectModifiedHandle );
+}
+
 void UDataprepParameterization::PostInitProperties()
 {
 	Super::PostInitProperties();
@@ -353,7 +364,8 @@ void UDataprepParameterization::PostLoad()
 {
 	if ( !HasAnyFlags( RF_ClassDefaultObject | RF_NeedLoad ) )
 	{
-		RegeneratedBindingAfterLoad();
+		SetFlags( RF_Public );
+		LoadParameterization();
 	}
 	Super::PostLoad();
 }
@@ -368,25 +380,50 @@ void UDataprepParameterization::Serialize(FArchive& Ar)
 	}
 	
 	Super::Serialize( Ar );
+}
+
+void UDataprepParameterization::PostEditUndo()
+{
+	// This implementation work on the assumption that all the objects in the transaction were serialized before the calls to post edit undo
+	PrepareCustomClassForNewClassGeneration();
+	UClass* OldClass = CustomContainerClass;
+	CustomContainerClass = nullptr;
+
+	LoadParameterization();
+
+	DoReinstancing( OldClass, false );
+
+	OnTellInstancesToReloadTheirSerializedData.Broadcast();
 
+	if ( UDataprepAsset* DataprepAsset = Cast<UDataprepAsset>( GetOuter() ) )
+	{
+		DataprepAsset->OnParameterizationWasModified.Broadcast();
+	}
+}
+
+void UDataprepParameterization::OnObjectModified(UObject* Object)
+{
+	if ( Object == DefaultParameterisation )
+	{
+		Modify();
+	}
 }
 
 UObject* UDataprepParameterization::GetDefaultObject()
 {
-	return CustomContainerClass->GetDefaultObject();
+	return DefaultParameterisation;
 }
 
 bool UDataprepParameterization::BindObjectProperty(UObject* InObject, const TArray<FDataprepPropertyLink>& InPropertyChain, FName ParameterName)
 {
-	if ( InObject && InPropertyChain.Num() > 0 )
+	if ( InObject && InPropertyChain.Num() > 0 && !ParameterName.IsNone() )
 	{
-		FDataprepParameterizationBinding Binding( InObject, InPropertyChain);
-
+		FDataprepParameterizationBinding Binding( InObject, InPropertyChain );
 		void* AddressOfTheValueFromBinding;
 
 		// We expect the chain to have a valid chain of cached property before inserting the binding
-		if ( DataprepParameterization::GetDeepestLevelOfValidCache(Binding, AddressOfTheValueFromBinding) == InPropertyChain.Num() - 1 )
-		{ 
+		if ( DataprepParameterization::GetDeepestLevelOfValidCache( Binding, AddressOfTheValueFromBinding) == InPropertyChain.Num() - 1 )
+		{
 			if ( UProperty** PropertyPtr = NameToParameterizationProperty.Find( ParameterName ) )
 			{
 				UProperty* Property = *PropertyPtr;
@@ -395,6 +432,8 @@ bool UDataprepParameterization::BindObjectProperty(UObject* InObject, const TArr
 				// Ensure that the properties type match
 				if ( Property->GetClass() == Binding.ValueType && !BindingsFromPipeline.Contains( Binding ) )
 				{
+					Modify();
+
 					// Valid bind the add binding to the name
 					BindingsFromPipeline.Add( Binding, ParameterName );
 					NameUsage.FindOrAdd( ParameterName ).Add( Binding );
@@ -402,8 +441,11 @@ bool UDataprepParameterization::BindObjectProperty(UObject* InObject, const TArr
 			}
 			else
 			{
+				Modify();
+				OnCustomClassAboutToBeUpdated.Broadcast();
+
 				BindingsFromPipeline.Add( MoveTemp(Binding) , ParameterName );
-				NameUsage.FindOrAdd( ParameterName ).Emplace( InObject, InPropertyChain );
+				NameUsage.FindOrAdd( ParameterName ).Add( FDataprepParameterizationBinding(InObject, InPropertyChain) );
 				UProperty* SourceProperty = InPropertyChain.Last().CachedProperty.Get();
 
 				// The validation we did with GetDeepestLevelOfValidCache ensure us that the property ptr is valid
@@ -423,11 +465,40 @@ bool UDataprepParameterization::BindObjectProperty(UObject* InObject, const TArr
 	return false;
 }
 
+bool UDataprepParameterization::IsObjectPropertyBinded(UObject* Object, const TArray<FDataprepPropertyLink>& PropertyChain) const
+{
+	FDataprepParameterizationBinding Binding( Object, PropertyChain );
+	return BindingsFromPipeline.Contains( Binding );
+}
+
+void UDataprepParameterization::RemoveBindedObjectProperty(UObject* Object, const TArray<FDataprepPropertyLink>& PropertyChain)
+{
+	FDataprepParameterizationBinding Binding( Object, PropertyChain );
+	uint32 Hash = GetTypeHash( Binding );
+	FName* ParameterNamePtr = BindingsFromPipeline.FindByHash( Hash, Binding );
+	if ( ParameterNamePtr )
+	{
+		Modify();
+		OnCustomClassAboutToBeUpdated.Broadcast();
+		BindingsFromPipeline.RemoveByHash( Hash, Binding );
+		FName ParameterName = *ParameterNamePtr; 
+		if ( TSet<FDataprepParameterizationBinding>* BindingAssociatedToName = NameUsage.Find( ParameterName ) )
+		{
+			BindingAssociatedToName->RemoveByHash( Hash, Binding );
+			if ( BindingAssociatedToName->Num() == 0 )
+			{
+				NameToParameterizationProperty.Remove( ParameterName );
+				UpdateClass();
+			}
+		}
+	}
+}
+
 void UDataprepParameterization::GenerateClass()
 {
 	if ( !CustomContainerClass )
 	{
-		CustomContainerClass = NewObject<UClass>(GetOutermost(), FName(TEXT("Parameterization")), RF_Transient);
+		CustomContainerClass = NewObject<UClass>(GetOutermost(), FName(TEXT("Parameterization")), RF_Transient );
 		CustomContainerClass->SetSuperStruct(UObject::StaticClass());
 
 		// Make the properties appear in a alphabetically order (for that we must add the properties to the class in the reverse order)
@@ -470,7 +541,7 @@ void UDataprepParameterization::UpdateClass()
 		CustomContainerClass->ClassFlags |= CLASS_NewerVersionExists;
 		CustomContainerClass->SetFlags( RF_NewerVersionExists );
 		CustomContainerClass->ClearFlags( RF_Public | RF_Standalone );
-		CustomContainerClass->Rename( *OldClassName, nullptr, REN_DontCreateRedirectors | REN_DoNotDirty);
+		CustomContainerClass->Rename( *OldClassName, nullptr, REN_DontCreateRedirectors | REN_DoNotDirty );
 	}
 
 	UClass* OldClass = CustomContainerClass;
@@ -479,48 +550,20 @@ void UDataprepParameterization::UpdateClass()
 
 	// Generate the new class
 	GenerateClass();
-	
-	if ( OldClass )
-	{
-		// For the CDO
-		UObject* OldCDO = OldClass->GetDefaultObject();
-		UObject* NewCDO = CustomContainerClass->GetDefaultObject();
-		GEngine->CopyPropertiesForUnrelatedObjects(OldClass->GetDefaultObject(), CustomContainerClass->GetDefaultObject());
-
-		// For the instances
-		TArray<UObject*> Objects;
-		constexpr bool bIncludeDerivedClasses = false;
-		GetObjectsOfClass(OldClass, Objects, bIncludeDerivedClasses);
-		TMap<UObject*, UObject*> OldToNew;
-		OldToNew.Reserve( Objects.Num() );
-		for (UObject* OldObject : Objects)
-		{
-			if ( OldObject && OldObject->IsValidLowLevel() )
-			{
-				FName ObjectName = OldObject->GetFName();
-				UObject* Outer = OldObject->GetOuter();
-				OldObject->Rename( nullptr, GetTransientPackage(), REN_DoNotDirty | REN_DontCreateRedirectors );
-				UObject* Object = NewObject<UObject>( Outer, CustomContainerClass, ObjectName, OldObject->GetFlags() );
-				GEngine->CopyPropertiesForUnrelatedObjects( OldObject, Object );
-				OldToNew.Add( OldObject, Object );
-			}
-		}
-		// The tools were already notify of the change by the copy properties for unrelated objects by GEngine::CopyPropertiesForUnrelatedObjects
-		OnClassUpdate.Broadcast( OldToNew );
 
-		DefaultParameterisation = NewCDO;
-	}
+	DoReinstancing( OldClass );
 }
 
-void UDataprepParameterization::RegeneratedBindingAfterLoad()
+void UDataprepParameterization::LoadParameterization()
 {
-	if (!CustomContainerClass)
-	{
+	if ( !CustomContainerClass )
+	{ 
 		CustomContainerClass = NewObject<UClass>( GetOutermost(), FName(TEXT("Parameterization")), RF_Transient );
 		CustomContainerClass->SetSuperStruct( UObject::StaticClass() );
 
 		TSet<FDataprepParameterizationBinding> BindingToRemove;
 		NameUsage.Empty( NameUsage.Num() );
+		NameToParameterizationProperty.Empty( NameToParameterizationProperty.Num() );
 
 		for ( TPair<FDataprepParameterizationBinding, FName>& Binding : BindingsFromPipeline )
 		{
@@ -562,18 +605,85 @@ void UDataprepParameterization::RegeneratedBindingAfterLoad()
 		CustomContainerClass->AssembleReferenceTokenStream(true);
 
 		DefaultParameterisation = CustomContainerClass->GetDefaultObject(true);
-
 		FDataprepParameterizationReader Reader( DefaultParameterisation, ParameterizationStorage );
 	}
 }
 
+void UDataprepParameterization::PrepareCustomClassForNewClassGeneration()
+{
+	if ( CustomContainerClass )
+	{
+		const FString OldClassName = MakeUniqueObjectName( CustomContainerClass->GetOuter(), CustomContainerClass->GetClass(), *FString::Printf(TEXT("%s_REINST"), *CustomContainerClass->GetName()) ).ToString();
+		CustomContainerClass->ClassFlags |= CLASS_NewerVersionExists;
+		CustomContainerClass->SetFlags( RF_NewerVersionExists );
+		CustomContainerClass->ClearFlags( RF_Public | RF_Standalone );
+		CustomContainerClass->Rename( *OldClassName, nullptr, REN_DontCreateRedirectors | REN_DoNotDirty);
+	}
+}
+
+void UDataprepParameterization::DoReinstancing(UClass* OldClass, bool bMigrateData)
+{
+	if ( OldClass && CustomContainerClass )
+	{
+		// For the CDO
+		UObject* OldCDO = OldClass->GetDefaultObject();
+		UObject* NewCDO = CustomContainerClass->GetDefaultObject();
+
+		if ( bMigrateData )
+		{
+			GEngine->CopyPropertiesForUnrelatedObjects( OldClass->GetDefaultObject(), CustomContainerClass->GetDefaultObject() );
+		}
+
+		// For the instances
+		TArray<UObject*> Objects;
+		constexpr bool bIncludeDerivedClasses = false;
+		GetObjectsOfClass( OldClass, Objects, bIncludeDerivedClasses );
+		TMap<UObject*, UObject*> OldToNew;
+		OldToNew.Reserve( Objects.Num() + 1 );
+		for (UObject* OldObject : Objects)
+		{
+			if ( OldObject && OldObject->IsValidLowLevel() )
+			{
+				FName ObjectName = OldObject->GetFName();
+				UObject* Outer = OldObject->GetOuter();
+				OldObject->Rename( nullptr, GetTransientPackage(), REN_DoNotDirty | REN_DontCreateRedirectors );
+
+				ObjectName = MakeUniqueObjectName( Outer, CustomContainerClass, ObjectName );
+				UObject* Object = NewObject<UObject>( Outer, CustomContainerClass, ObjectName, OldObject->GetFlags() );
+				
+				if ( bMigrateData )
+				{
+					GEngine->CopyPropertiesForUnrelatedObjects( OldObject, Object );
+				}
+
+				OldToNew.Add( OldObject, Object );
+			}
+		}
+
+		OldToNew.Add( OldCDO, NewCDO );
+
+		/**
+		 * Notify the tools 
+		 * If we did the data migration the tools were already notify of the change by the copy properties for unrelated objects by GEngine::CopyPropertiesForUnrelatedObjects
+		 */
+		if ( !bMigrateData && GEngine)
+		{
+			GEngine->NotifyToolsOfObjectReplacement(OldToNew);
+		}
+
+		OnCustomClassWasUpdated.Broadcast( OldToNew );
+
+		DefaultParameterisation = NewCDO;
+	}
+}
+
 UProperty* UDataprepParameterization::AddPropertyToClass(FName ParameterisationPropertyName, UProperty& Property)
 {
 	if ( !NameToParameterizationProperty.Find( ParameterisationPropertyName ) )
 	{
 		UProperty* NewProperty = DuplicateObject<UProperty>( &Property, CustomContainerClass, ParameterisationPropertyName );
 		NewProperty->SetFlags( RF_Transient );
-		NewProperty->PropertyFlags = CPF_Edit;
+		NewProperty->PropertyFlags = CPF_Edit | CPF_NonTransactional;
 
 		// Need to manually call Link to fix-up some data (such as the C++ property flags) that are only set during Link
 		{
@@ -589,27 +699,52 @@ UProperty* UDataprepParameterization::AddPropertyToClass(FName ParameterisationP
 	return nullptr;
 }
 
+
+UDataprepParameterizationInstance::UDataprepParameterizationInstance()
+{
+	OnObjectModifiedHandle = FCoreUObjectDelegates::OnObjectModified.AddUObject(this, &UDataprepParameterizationInstance::OnObjectModified );
+}
+
+UDataprepParameterizationInstance::~UDataprepParameterizationInstance()
+{
+	FCoreUObjectDelegates::OnObjectModified.Remove( OnObjectModifiedHandle );
+}
+
 void UDataprepParameterizationInstance::PostLoad()
 {
-	if (!HasAnyFlags(RF_ClassDefaultObject | RF_NeedLoad))
+	if ( !HasAnyFlags( RF_ClassDefaultObject | RF_NeedLoad ) )
 	{
 		SetFlags( RF_Public );
 		LoadParameterization();
+		SetupCallbacksFromSourceParameterisation();
 	}
 	Super::PostLoad();
 }
 
 void UDataprepParameterizationInstance::Serialize(FArchive& Ar)
 {
-	if (Ar.IsSaving() && !HasAnyFlags(RF_ClassDefaultObject))
+	if ( Ar.IsSaving() && !HasAnyFlags( RF_ClassDefaultObject ) )
 	{
 		check( SourceParameterization );
 		// Todo track when the object is changed to avoid rewriting to array each time
-		ParameterizationInstanceStorage.Empty( SourceParameterization->CustomContainerClass->GetMinAlignment() );
+		ParameterizationInstanceStorage.Empty( ParameterizationInstanceStorage.Num() );
 		FDataprepParameterizationWriter Writer( ParameterizationInstance, ParameterizationInstanceStorage );
 	}
 
-	Super::Serialize(Ar);
+	Super::Serialize( Ar );
+}
+
+void UDataprepParameterizationInstance::PostEditUndo()
+{
+	LoadParameterization();
+}
+
+void UDataprepParameterizationInstance::OnObjectModified(UObject* Object)
+{
+	if ( Object == ParameterizationInstance )
+	{
+		Modify();
+	}
 }
 
 void UDataprepParameterizationInstance::ApplyParameterization(const TMap<UObject*, UObject*>& SourceToCopy)
@@ -634,7 +769,13 @@ void UDataprepParameterizationInstance::ApplyParameterization(const TMap<UObject
 	}
 }
 
-void UDataprepParameterizationInstance::OnCustomClassUpdate(const TMap<UObject*, UObject*>& OldToNew)
+void UDataprepParameterizationInstance::CustomClassAboutToBeUpdated()
+{
+	// The instance is about to be modified
+	Modify();
+}
+
+void UDataprepParameterizationInstance::CustomClassWasUpdated(const TMap<UObject*, UObject*>& OldToNew)
 {
 	if ( UObject* NewInstance = OldToNew.FindRef( ParameterizationInstance ) )
 	{
@@ -648,22 +789,44 @@ void UDataprepParameterizationInstance::LoadParameterization()
 
 	if ( !SourceParameterization->CustomContainerClass )
 	{
-		SourceParameterization->RegeneratedBindingAfterLoad();
+		SourceParameterization->LoadParameterization();
 	}
 
-	SourceParameterization->OnClassUpdate.AddUObject( this, &UDataprepParameterizationInstance::OnCustomClassUpdate );
+	if ( !ParameterizationInstance )
+	{
+		ParameterizationInstance = NewObject<UObject>( this, SourceParameterization->CustomContainerClass, FName(TEXT("Parameterization")), RF_Transient );
+	}
 
-	ParameterizationInstance = NewObject<UObject>( this, SourceParameterization->CustomContainerClass, FName(TEXT("Parameterization")), RF_Transient | RF_Transactional );
 	FDataprepParameterizationReader Reader( ParameterizationInstance, ParameterizationInstanceStorage );
 }
 
 void UDataprepParameterizationInstance::SetParameterizationSource(UDataprepParameterization& Parameterization)
 {
-	if ( SourceParameterization )
-	{
-		SourceParameterization->OnClassUpdate.RemoveAll( this );
-	}
+	UndoSetupForCallbacksFromParameterization();
 
 	SourceParameterization = &Parameterization;
+	SetupCallbacksFromSourceParameterisation();
+
+	// Reload the parameterization (this act as a sort of data migration process)
 	LoadParameterization();
 }
+
+void UDataprepParameterizationInstance::SetupCallbacksFromSourceParameterisation()
+{
+	check( SourceParameterization );
+
+	SourceParameterization->OnCustomClassAboutToBeUpdated.AddUObject( this, &UDataprepParameterizationInstance::CustomClassAboutToBeUpdated );
+	SourceParameterization->OnCustomClassWasUpdated.AddUObject( this, &UDataprepParameterizationInstance::CustomClassWasUpdated );
+	SourceParameterization->OnTellInstancesToReloadTheirSerializedData.AddUObject( this, &UDataprepParameterizationInstance::LoadParameterization );
+}
+
+void UDataprepParameterizationInstance::UndoSetupForCallbacksFromParameterization()
+{
+	if (SourceParameterization)
+	{
+		SourceParameterization->OnCustomClassAboutToBeUpdated.RemoveAll( this );
+		SourceParameterization->OnCustomClassWasUpdated.RemoveAll( this );
+		SourceParameterization->OnTellInstancesToReloadTheirSerializedData.RemoveAll( this );
+	}
+}
+
diff --git a/Engine/Plugins/Experimental/DataPrepEditor/Source/DataPrepCore/Private/Parameterization/DataprepParameterization.h b/Engine/Plugins/Experimental/DataPrepEditor/Source/DataPrepCore/Private/Parameterization/DataprepParameterization.h
--- a/Engine/Plugins/Experimental/DataPrepEditor/Source/DataPrepCore/Private/Parameterization/DataprepParameterization.h
+++ b/Engine/Plugins/Experimental/DataPrepEditor/Source/DataPrepCore/Private/Parameterization/DataprepParameterization.h
@@ -6,6 +6,7 @@
 
 #include "CoreMinimal.h"
 #include "Delegates/DelegateCombinations.h"
+#include "Delegates/IDelegateInstance.h"
 #include "Engine/UserDefinedStruct.h"
 #include "Templates/SharedPointer.h"
 #include "UObject/Class.h"
@@ -54,6 +55,18 @@ uint32 GetTypeHash(const FDataprepParameterizationBinding& Binding);
 
 uint32 GetTypeHash(const TArray<FDataprepPropertyLink>& PropertyLinks);
 
+/**
+ * Count the number of time each hash was encounter
+ */
+USTRUCT()
+struct FHashCount
+{
+	GENERATED_BODY()
+
+	UPROPERTY()
+	TMap<uint32, uint32> HashCount;
+};
+
 /** 
  * The DataprepParameterization contains the data for the parameterization of a pipeline
  */
@@ -63,24 +76,55 @@ class UDataprepParameterization : public UObject
 public:
 	GENERATED_BODY()
 
+	UDataprepParameterization();
+	~UDataprepParameterization();
+
 	// UObject interface
 	virtual void PostInitProperties() override;
 	virtual void PostLoad() override;
 	virtual void Serialize(FArchive& Ar) override;
+	virtual void PostEditUndo() override;
 	// End of UObject interface
 
-	// Temporary function to help debug and test the feature
+	void OnObjectModified(UObject* Object);
+
+
 	UObject* GetDefaultObject();
 
 	bool BindObjectProperty(UObject* Object, const TArray<FDataprepPropertyLink>& PropertyChain, FName Name);
 
+	bool IsObjectPropertyBinded(UObject* Object, const TArray<FDataprepPropertyLink>& PropertyChain) const;
+
+	void RemoveBindedObjectProperty(UObject* Object, const TArray<FDataprepPropertyLink>& PropertyChain);
+
 private:
 
+	/**
+	 * Generate the Custom Container Class
+	 */
 	void GenerateClass();
 
+	/**
+	 * Update the Custom Container Class to a newer version
+	 */
 	void UpdateClass();
+	
+	/**
+	 * Do the process of regenerating the Custom Container Class and the data of its default object from the serialized data
+	 */
+	void LoadParameterization();
+
+	/**
+	 * Remove the current Custom Container Class so that we can create a new one
+	 */
+	void PrepareCustomClassForNewClassGeneration();
 
-	void RegeneratedBindingAfterLoad();
+	/**
+	 * Do reinstancing of the objects created from the Custom Container Class
+	 * @param OldClass The previous Custom Constainer Class
+	 * @param bMigrateData Should we migrate the data from the old instances to the new instances
+	 */
+	void DoReinstancing(UClass* OldClass, bool bMigrateData = true);
 
 	/**
 	 * Try adding a binded property to the parameterization class
@@ -91,30 +135,38 @@ private:
 	UPROPERTY()
 	TMap<FDataprepParameterizationBinding, FName> BindingsFromPipeline;
 
-	UPROPERTY(Transient)
+	UPROPERTY(Transient, NonTransactional)
 	TMap<FName, UProperty*> NameToParameterizationProperty; // Just a cache for the CustomFindProperty
 
 	/** Track the name usage for parameters */
-	TMap<FName, TArray<FDataprepParameterizationBinding>> NameUsage;
+	TMap<FName, TSet<FDataprepParameterizationBinding>> NameUsage;
 
-	UPROPERTY(Transient)
+	UPROPERTY(Transient, NonTransactional)
 	UClass* CustomContainerClass;
 
-	UPROPERTY(Transient)
+	UPROPERTY(Transient, NonTransactional)
 	UObject* DefaultParameterisation;
 
 	/** 
-	 * This is used only to store a serialization of the values of the parameterization since we can't save our custom class
+	 * This is used only to store a serialization of the values of the parameterization since we can't save our custom container class
 	 */
 	UPROPERTY()
 	TArray<uint8> ParameterizationStorage;
 
+	DECLARE_EVENT(UDataprepParameterization, FOnCustomClassAboutToBeUpdated);
+	FOnCustomClassAboutToBeUpdated OnCustomClassAboutToBeUpdated;
+
 	using FMapOldToNewObjects = TMap<UObject*, UObject*>;
-	DECLARE_EVENT_OneParam(UDataprepParameterization, FCustomClassWasUpdated, const FMapOldToNewObjects& /** OldToNew */);
-	FCustomClassWasUpdated OnClassUpdate;
+	DECLARE_EVENT_OneParam(UDataprepParameterization, FOnCustomClassWasUpdated, const FMapOldToNewObjects& /** OldToNew */);
+	FOnCustomClassWasUpdated OnCustomClassWasUpdated;
+
+	DECLARE_EVENT(UDataprepParameterization,FOnTellInstancesToReloadTheirSerializedData);
+	FOnTellInstancesToReloadTheirSerializedData OnTellInstancesToReloadTheirSerializedData;
 
 	// the dataprep instance need some special access to the dataprep parameterization
 	friend class UDataprepParameterizationInstance;
+
+	FDelegateHandle OnObjectModifiedHandle;
 };
 
 
@@ -124,11 +176,17 @@ class UDataprepParameterizationInstance : public UObject
 public:
 	GENERATED_BODY()
 
+	UDataprepParameterizationInstance();
+	~UDataprepParameterizationInstance();
+
 	// UObject interface
 	virtual void PostLoad() override;
 	virtual void Serialize(FArchive& Ar) override;
+	virtual void PostEditUndo() override;
 	// End of UObject interface
 
+	void OnObjectModified(UObject* Object);
+
 	// Apply the parameterization to a copy of the source pipeline
 	void ApplyParameterization(const TMap<UObject*, UObject*>& SourceToCopy);
 
@@ -138,20 +196,41 @@ public:
 
 private:
 
-	void OnCustomClassUpdate(const TMap<UObject*, UObject*>& OldToNew);
+	void CustomClassAboutToBeUpdated();
+
+	/**
+	 * Used as call back for event coming from the source parameterization
+	 * Change the parametrization instance to the new object after a reinstancing
+	 */
+	void CustomClassWasUpdated(const TMap<UObject*, UObject*>& OldToNew);
 
+	/**
+	 * Load the parameterization data on the instance from the ParameterizationInstanceStorage
+	 */
 	void LoadParameterization();
 
+	/**
+	 * Setup the parameterization instance so that we can react to event coming from the source parameterization
+	 */
+	void SetupCallbacksFromSourceParameterisation();
+
+	/**
+	 * Clean the parameterization instance so that we can bind to a new source parameterization
+	 */
+	void UndoSetupForCallbacksFromParameterization();
+
+	// The parameterization from which this instance was constructed
 	UPROPERTY()
 	UDataprepParameterization* SourceParameterization;
 
-	UPROPERTY(Transient)
+	// The actual object on which the parameterization data is stored
+	UPROPERTY(Transient, NonTransactional)
 	UObject* ParameterizationInstance;
 
-	/** 
-	 * This is used only to store a serialization of the values of the parameterization since we can't save the custom class
-	 */
+	// This is used only to store a serialization of the values of the parameterization since we can't save the custom class
 	UPROPERTY()
 	TArray<uint8> ParameterizationInstanceStorage;
+
+	FDelegateHandle OnObjectModifiedHandle;
 };
 
diff --git a/Engine/Plugins/Experimental/DataPrepEditor/Source/DataPrepCore/Private/Parameterization/DataprepParameterizationUtils.cpp b/Engine/Plugins/Experimental/DataPrepEditor/Source/DataPrepCore/Private/Parameterization/DataprepParameterizationUtils.cpp
--- a/Engine/Plugins/Experimental/DataPrepEditor/Source/DataPrepCore/Private/Parameterization/DataprepParameterizationUtils.cpp
+++ b/Engine/Plugins/Experimental/DataPrepEditor/Source/DataPrepCore/Private/Parameterization/DataprepParameterizationUtils.cpp
@@ -103,6 +103,25 @@ TArray<FDataprepPropertyLink> FDataprepParameterizationUtils::MakePropertyChain(
 	return PropertyChain;
 }
 
+FDataprepParameterizationContext FDataprepParameterizationUtils::CreateContext(TSharedPtr<IPropertyHandle> PropertyHandle, const FDataprepParameterizationContext& ParameterisationContext)
+{
+	FDataprepParameterizationContext NewContext;
+	NewContext.State = ParameterisationContext.State;
+
+	if ( NewContext.State == EParametrizationState::CanBeParameterized || NewContext.State == EParametrizationState::InvalidForParameterization )
+	{
+		// This implementation could be improved by incrementally expending the property chain.
+		NewContext.PropertyChain = MakePropertyChain( PropertyHandle );
+		NewContext.State = NewContext.PropertyChain.Num() > 0 ? EParametrizationState::CanBeParameterized : EParametrizationState::InvalidForParameterization;
+	}
+	else if ( NewContext.State == EParametrizationState::IsParameterized )
+	{
+		NewContext.State = EParametrizationState::ParentIsParameterized;
+	}
+
+	return NewContext;
+}
+
 UDataprepAsset* FDataprepParameterizationUtils::GetDataprepAssetForParameterization(UObject* Object)
 {
 	if ( Object )
diff --git a/Engine/Plugins/Experimental/DataPrepEditor/Source/DataPrepCore/Public/DataPrepAsset.h b/Engine/Plugins/Experimental/DataPrepEditor/Source/DataPrepCore/Public/DataPrepAsset.h
--- a/Engine/Plugins/Experimental/DataPrepEditor/Source/DataPrepCore/Public/DataPrepAsset.h
+++ b/Engine/Plugins/Experimental/DataPrepEditor/Source/DataPrepCore/Public/DataPrepAsset.h
@@ -20,7 +20,6 @@ struct FDataprepActionContext;
 struct FDataprepConsumerContext;
 struct FDataprepProducerContext;
 
-
 /**
  * A DataprepAsset is an implementation of the DataprepAssetInterface using
  * a Blueprint as the recipe pipeline. The Blueprint is composed of DataprepAction
@@ -92,13 +91,22 @@ public:
 	FOnDataprepBlueprintChange& GetOnBlueprintChanged() { return OnBlueprintChanged; }
 	// end of temp code for nodes development
 
+
 	/**
-	 * Temporary function to get to allow the editor to view the default parameterization
+	 * Delegate to notify the ui that a dataprep parametrization was modified
+	 * This necessary as the ui for the parameterization is only updated by manual event
 	 */
+	DECLARE_MULTICAST_DELEGATE(FDataprepParameterizationWasModified)
+	FDataprepParameterizationWasModified OnParameterizationWasModified;
+
 	virtual UObject* GetParameterizationObject() override;
 
 	void BindObjectPropertyToParameterization(UObject* Object, const TArray<struct FDataprepPropertyLink>& InPropertyChain, FName Name);
 
+	bool IsObjectPropertyBinded(UObject* Object, const TArray<struct FDataprepPropertyLink>& InPropertyChain) const;
+
+	void RemoveObjectPropertyFromParameterization(UObject* Object, const TArray<struct FDataprepPropertyLink>& InPropertyChain);
+
 	UDataprepParameterization* GetDataprepParameterization() { return Parameterization; }
 
 protected:
diff --git a/Engine/Plugins/Experimental/DataPrepEditor/Source/DataPrepCore/Public/Parameterization/DataprepParameterizationUtils.h b/Engine/Plugins/Experimental/DataPrepEditor/Source/DataPrepCore/Public/Parameterization/DataprepParameterizationUtils.h
--- a/Engine/Plugins/Experimental/DataPrepEditor/Source/DataPrepCore/Public/Parameterization/DataprepParameterizationUtils.h
+++ b/Engine/Plugins/Experimental/DataPrepEditor/Source/DataPrepCore/Public/Parameterization/DataprepParameterizationUtils.h
@@ -12,6 +12,24 @@ class IPropertyHandle;
 class UDataprepAsset;
 class UProperty;
 
+enum class EParametrizationState : uint8
+{
+	CanBeParameterized,
+	IsParameterized,
+	ParentIsParameterized,
+	InvalidForParameterization
+};
+
+/**
+ * A small context that help when constructing the widgets for the parameterization
+ */
+struct FDataprepParameterizationContext
+{
+	TArray<FDataprepPropertyLink> PropertyChain;
+	EParametrizationState State = EParametrizationState::CanBeParameterized;
+};
+
+
 USTRUCT()
 struct FDataprepPropertyLink
 {
@@ -51,6 +69,12 @@ public:
 	 */
 	static TArray<FDataprepPropertyLink> MakePropertyChain(TSharedPtr<IPropertyHandle> PropertyHandle);
 
+
+	/**
+	 * Take a already existing parameterization context and create a new version including the handle.
+	 */
+	static FDataprepParameterizationContext CreateContext(TSharedPtr<IPropertyHandle> PropertyHandle, const FDataprepParameterizationContext& ParameterisationContext);
+
 	/**
 	 * Grab the dataprep used for the parameterization of the object.
 	 * @param Object The object from which the bindings would be created
diff --git a/Engine/Plugins/Experimental/DataPrepEditor/Source/DataPrepEditor/Private/Widgets/DataprepWidgets.cpp b/Engine/Plugins/Experimental/DataPrepEditor/Source/DataPrepEditor/Private/Widgets/DataprepWidgets.cpp
--- a/Engine/Plugins/Experimental/DataPrepEditor/Source/DataPrepEditor/Private/Widgets/DataprepWidgets.cpp
+++ b/Engine/Plugins/Experimental/DataPrepEditor/Source/DataPrepEditor/Private/Widgets/DataprepWidgets.cpp
@@ -13,18 +13,19 @@
 #include "Dialogs/DlgPickPath.h"
 #include "Editor.h"
 #include "EditorFontGlyphs.h"
+#include "EditorFontGlyphs.h"
 #include "EditorStyleSet.h"
 #include "Framework/MultiBox/MultiBoxBuilder.h"
 #include "HAL/PlatformProcess.h"
 #include "IContentBrowserSingleton.h"
 #include "IDetailTreeNode.h"
 #include "IPropertyRowGenerator.h"
-#include "EditorFontGlyphs.h"
 #include "Misc/Paths.h"
 #include "Modules/ModuleManager.h"
 #include "PropertyCustomizationHelpers.h"
 #include "PropertyEditorModule.h"
 #include "PropertyEditorModule.h"
+#include "ScopedTransaction.h"
 #include "Styling/CoreStyle.h"
 #include "UObject/StrongObjectPtr.h"
 #include "UObject/UnrealType.h"
@@ -344,37 +345,66 @@ void SDataprepConsumerWidget::UpdateContentFolderText()
 }
 
 
-TSharedRef< SWidget > SDataprepDetailsView::CreateDefaultWidget( TSharedPtr< SWidget >& NameWidget, TSharedPtr< SWidget >& ValueWidget, float LeftPadding, EHorizontalAlignment HAlign, EVerticalAlignment VAlign, const TArray<FDataprepPropertyLink>& InPropertyChain)
+TSharedRef< SWidget > SDataprepDetailsView::CreateDefaultWidget( TSharedPtr< SWidget >& NameWidget, TSharedPtr< SWidget >& ValueWidget, float LeftPadding, EHorizontalAlignment HAlign, EVerticalAlignment VAlign, const FDataprepParameterizationContext& ParameterizationContext)
 {
 	TSharedRef<SHorizontalBox> NameColumn = SNew(SHorizontalBox)
 		.Clipping(EWidgetClipping::OnDemand);
 
 	// Optionally add the parameterization widget
-	if ( DataprepAssetForParameterization.IsValid() && InPropertyChain.Num() > 0 )
+	if ( DataprepAssetForParameterization.IsValid() )
 	{
+		if ( ParameterizationContext.State == EParametrizationState::CanBeParameterized )
+		{
+			NameColumn->AddSlot()
+				.HAlign(HAlign_Left)
+				.VAlign(VAlign_Center)
+				.AutoWidth()
+				[
 
+					SNew(STextBlock)
+					.Font(FEditorStyle::Get().GetFontStyle("FontAwesome.11"))
+					.Text(FEditorFontGlyphs::Plus_Circle)
+					.Margin(FMargin(5.0f, 5.0f, 0.0f, 0.0f))
+					.ColorAndOpacity(FColor::Blue)
+					.OnDoubleClicked_Lambda([this, PropertyChain = ParameterizationContext.PropertyChain](const FGeometry&,const FPointerEvent&) -> FReply
+						{
+							if ( UDataprepAsset* DataprepAsset = DataprepAssetForParameterization.Get() )
+							{
+								FScopedTransaction Transaction( LOCTEXT("DataprepBindingToParameterization","Adding Parameter") );
+								DataprepAssetForParameterization->BindObjectPropertyToParameterization( DetailedObject, PropertyChain, *PropertyChain.Last().CachedProperty->GetDisplayNameText().ToString() );
+								bRefreshObjectToDisplay = true;
+							}
+
+							return FReply::Handled();
+						})
+				];
+		}
+		else if ( ParameterizationContext.State == EParametrizationState::IsParameterized )
+		{
+			NameColumn->AddSlot()
+				.HAlign(HAlign_Left)
+				.VAlign(VAlign_Center)
+				.AutoWidth()
+				[
 
-		NameColumn->AddSlot()
-			.HAlign(HAlign_Left)
-			.VAlign(VAlign_Center)
-			.AutoWidth()
-			[
-
-				SNew(STextBlock)
-				.Font(FEditorStyle::Get().GetFontStyle("FontAwesome.11"))
-				.Text(FEditorFontGlyphs::Plus_Circle)
-				.Margin(FMargin(5.0f, 5.0f, 0.0f, 0.0f))
-				.ColorAndOpacity(FColor::Blue)
-				.OnDoubleClicked_Lambda([this, PropertyChain = InPropertyChain](const FGeometry&,const FPointerEvent&) -> FReply
+					SNew(STextBlock)
+					.Font(FEditorStyle::Get().GetFontStyle("FontAwesome.11"))
+					.Text(FEditorFontGlyphs::Minus_Circle)
+					.Margin(FMargin(5.0f, 5.0f, 0.0f, 0.0f))
+					.ColorAndOpacity(FColor::Red)
+					.OnDoubleClicked_Lambda([this, PropertyChain = ParameterizationContext.PropertyChain](const FGeometry&, const FPointerEvent&)->FReply
 					{
-						if ( UDataprepAsset* DataprepAsset = DataprepAssetForParameterization.Get() )
+						if (UDataprepAsset * DataprepAsset = DataprepAssetForParameterization.Get())
 						{
-							DataprepAssetForParameterization->BindObjectPropertyToParameterization( DetailedObject, PropertyChain, *PropertyChain.Last().CachedProperty->GetDisplayNameText().ToString() );
+							FScopedTransaction Transaction( LOCTEXT("DataprepRemoveBindingFromParameterization","Removing Parameter") );
+							DataprepAssetForParameterization->RemoveObjectPropertyFromParameterization( DetailedObject, PropertyChain );
+							bRefreshObjectToDisplay = true;
 						}
 
 						return FReply::Handled();
 					})
-			];
+				];
+		}
 	}
 
 	// Add the name widget
@@ -432,7 +462,7 @@ void SDataprepDetailsView::OnPropertyChanged(const FPropertyChangedEvent& InEven
 
 	if( TrackedProperties.Contains( InEvent.Property ) )
 	{
-		bNewObjectToDisplay = true;
+		bRefreshObjectToDisplay = true;
 	}
 }
 
@@ -441,11 +471,16 @@ void SDataprepDetailsView::OnObjectReplaced(const TMap<UObject*, UObject*>& Repl
 	if ( UObject * const* ObjectPtr = ReplacementObjectMap.Find( DetailedObject ) )
 	{
 		DetailedObject = *ObjectPtr;
-		bNewObjectToDisplay = true;
+		bRefreshObjectToDisplay = true;
 	}
 }
 
-void SDataprepDetailsView::AddWidgets( const TArray< TSharedRef< IDetailTreeNode > >& DetailTree, TSharedPtr<SGridPanel>& GridPanel, int32& Index, float LeftPadding )
+void SDataprepDetailsView::ForceRefresh()
+{
+	bRefreshObjectToDisplay = true;
+}
+
+void SDataprepDetailsView::AddWidgets( const TArray< TSharedRef< IDetailTreeNode > >& DetailTree, TSharedPtr<SGridPanel>& GridPanel, int32& Index, float LeftPadding, const FDataprepParameterizationContext& InParameterizationContext)
 {
 	auto IsDetailNodeDisplayable = []( const TSharedPtr< IPropertyHandle >& PropertyHandle)
 	{
@@ -493,6 +528,18 @@ void SDataprepDetailsView::AddWidgets( const TArray< TSharedRef< IDetailTreeNode
 	for( const TSharedRef< IDetailTreeNode >& ChildNode : DetailTree )
 	{
 		TSharedPtr< IPropertyHandle > PropertyHandle = ChildNode->CreatePropertyHandle();
+		FDataprepParameterizationContext CurrentParameterizationContext = FDataprepParameterizationUtils::CreateContext( PropertyHandle, InParameterizationContext );
+		if ( CurrentParameterizationContext.State == EParametrizationState::CanBeParameterized )
+		{
+			if ( UDataprepAsset* DataprepAsset = DataprepAssetForParameterization.Get() )
+			{
+				if ( DataprepAsset->IsObjectPropertyBinded( DetailedObject, CurrentParameterizationContext.PropertyChain ) )
+				{
+					CurrentParameterizationContext.State = EParametrizationState::IsParameterized;
+				}
+			}
+		}
+
 
 		if( ChildNode->GetNodeType() == EDetailNodeType::Category )
 		{
@@ -514,7 +561,7 @@ void SDataprepDetailsView::AddWidgets( const TArray< TSharedRef< IDetailTreeNode
 
 			TArray< TSharedRef< IDetailTreeNode > > Children;
 			ChildNode->GetChildren( Children );
-			AddWidgets( Children, GridPanel, Index, LeftPadding );
+			AddWidgets( Children, GridPanel, Index, LeftPadding, CurrentParameterizationContext );
 		}
 		else if( IsDetailNodeDisplayableContainerProperty( PropertyHandle ) )
 		{
@@ -527,7 +574,7 @@ void SDataprepDetailsView::AddWidgets( const TArray< TSharedRef< IDetailTreeNode
 				DetailPropertyRow->GetDefaultWidgets( NameWidget, ValueWidget, Row, true );
 				GridPanel->AddSlot(0, Index)
 				[
-					CreateDefaultWidget( NameWidget, ValueWidget, LeftPadding, Row.ValueWidget.HorizontalAlignment, Row.ValueWidget.VerticalAlignment, FDataprepParameterizationUtils::MakePropertyChain( PropertyHandle ) )
+					CreateDefaultWidget( NameWidget, ValueWidget, LeftPadding, Row.ValueWidget.HorizontalAlignment, Row.ValueWidget.VerticalAlignment, CurrentParameterizationContext )
 				];
 
 				Index++;
@@ -537,7 +584,7 @@ void SDataprepDetailsView::AddWidgets( const TArray< TSharedRef< IDetailTreeNode
 				if( Children.Num() > 0 )
 				{
 					// #ueent_todo: Find a way to add collapse/expand capability for property with children
-					AddWidgets( Children, GridPanel, Index, LeftPadding + 10.f );
+					AddWidgets( Children, GridPanel, Index, LeftPadding + 10.f, CurrentParameterizationContext );
 				}
 
 				TrackedProperties.Add( PropertyHandle->GetProperty() );
@@ -572,7 +619,7 @@ void SDataprepDetailsView::AddWidgets( const TArray< TSharedRef< IDetailTreeNode
 			{
 				GridPanel->AddSlot(0, Index)
 				[
-					CreateDefaultWidget( NameWidget, ValueWidget, LeftPadding, HAlign, VAlign, FDataprepParameterizationUtils::MakePropertyChain( PropertyHandle ) )
+					CreateDefaultWidget( NameWidget, ValueWidget, LeftPadding, HAlign, VAlign, CurrentParameterizationContext )
 				];
 				Index++;
 
@@ -593,7 +640,7 @@ void SDataprepDetailsView::AddWidgets( const TArray< TSharedRef< IDetailTreeNode
 				if( bDisplayChildren && Children.Num() > 0 )
 				{
 					// #ueent_todo: Find a way to add collapse/expand capability for property with children
-					AddWidgets( Children, GridPanel, Index, LeftPadding + 10.f );
+					AddWidgets( Children, GridPanel, Index, LeftPadding + 10.f, CurrentParameterizationContext );
 				}
 			}
 		}
@@ -602,7 +649,7 @@ void SDataprepDetailsView::AddWidgets( const TArray< TSharedRef< IDetailTreeNode
 
 void SDataprepDetailsView::Construct(const FArguments& InArgs)
 {
-	bNewObjectToDisplay = false;
+	bRefreshObjectToDisplay = false;
 	DetailedObject = InArgs._Object;
 
 	if (InArgs._ColumnSizeData.IsValid())
@@ -646,14 +693,21 @@ void SDataprepDetailsView::Construct()
 
 	if ( DetailedObject )
 	{
-		DataprepAssetForParameterization = FDataprepParameterizationUtils::GetDataprepAssetForParameterization( DetailedObject );
+		UDataprepAsset* DataprepAsset = FDataprepParameterizationUtils::GetDataprepAssetForParameterization( DetailedObject );
+		if ( DataprepAsset )
+		{
+			OnDataprepParameterizationWasModifiedHandle = DataprepAsset->OnParameterizationWasModified.AddSP( this, &SDataprepDetailsView::ForceRefresh );
+		}
 
-		TSharedPtr<SGridPanel> GridPanel = SNew(SGridPanel).FillColumn( 0.0f, 1.0f );
+		FDataprepParameterizationContext ParameterizationContext;
+		ParameterizationContext.State = DataprepAsset ? EParametrizationState::CanBeParameterized : EParametrizationState::InvalidForParameterization;
+		DataprepAssetForParameterization = DataprepAsset;
 
+		TSharedPtr<SGridPanel> GridPanel = SNew(SGridPanel).FillColumn( 0.0f, 1.0f );
 		TArray< TSharedRef< IDetailTreeNode > > RootNodes = Generator->GetRootTreeNodes();
 
 		int32 Index = 0;
-		AddWidgets(RootNodes, GridPanel, Index, 0.f);
+		AddWidgets(RootNodes, GridPanel, Index, 0.f, ParameterizationContext);
 
 		ChildSlot
 		[
@@ -700,20 +754,25 @@ SDataprepDetailsView::~SDataprepDetailsView()
 	{
 		GEditor->OnObjectsReplaced().Remove( OnObjectReplacedHandle );
 	}
+
+	if ( UDataprepAsset* DataprepAsset = DataprepAssetForParameterization.Get() )
+	{
+		DataprepAsset->OnParameterizationWasModified.Remove( OnDataprepParameterizationWasModifiedHandle );
+	}
 }
 
 void SDataprepDetailsView::Tick(const FGeometry& AllottedGeometry, const double InCurrentTime, const float InDeltaTime)
 {
 	Super::Tick( AllottedGeometry, InCurrentTime, InDeltaTime );
 
-	if ( bNewObjectToDisplay )
+	if ( bRefreshObjectToDisplay )
 	{
 		TArray< UObject* > Objects;
 		Objects.Add( DetailedObject );
 		Generator->SetObjects( Objects );
 		 
 		Construct();
-		bNewObjectToDisplay = false;
+		bRefreshObjectToDisplay = false;
 	}
 }
 
@@ -723,7 +782,7 @@ void SDataprepDetailsView::SetObjectToDisplay(UObject& Object)
 	if ( DetailedObject != NewObjectToDisplay )
 	{
 		DetailedObject = NewObjectToDisplay;
-		bNewObjectToDisplay = true;
+		bRefreshObjectToDisplay = true;
 	}
 }
 
diff --git a/Engine/Plugins/Experimental/DataPrepEditor/Source/DataPrepEditor/Private/Widgets/DataprepWidgets.h b/Engine/Plugins/Experimental/DataPrepEditor/Source/DataPrepEditor/Private/Widgets/DataprepWidgets.h
--- a/Engine/Plugins/Experimental/DataPrepEditor/Source/DataPrepEditor/Private/Widgets/DataprepWidgets.h
+++ b/Engine/Plugins/Experimental/DataPrepEditor/Source/DataPrepEditor/Private/Widgets/DataprepWidgets.h
@@ -11,6 +11,8 @@
 class SEditableTextBox;
 class UDataprepAsset;
 class UDataprepContentConsumer;
+
+struct FDataprepParameterizationContext;
 struct FDataprepPropertyLink;
 
 struct FDataprepDetailsViewColumnSizeData
@@ -102,7 +104,7 @@ private:
 	void Construct();
 
 	/** Add widgets held by array of DetailTreeNode objects */
-	void AddWidgets( const TArray< TSharedRef< class IDetailTreeNode > >& DetailTree, TSharedPtr< class SGridPanel >& GridPanel, int32& Index, float LeftPadding);
+	void AddWidgets( const TArray< TSharedRef< class IDetailTreeNode > >& DetailTree, TSharedPtr< class SGridPanel >& GridPanel, int32& Index, float LeftPadding, const FDataprepParameterizationContext& ParameterizationContext);
 
 	/**
 	 * Create generic widget for a property row
@@ -112,7 +114,7 @@ private:
 	 * @param PropertyChain The chain to the property from the class of the detailed object
 	 * @return A row for the details view
 	 */
-	TSharedRef< SWidget > CreateDefaultWidget(TSharedPtr< SWidget >& NameWidget, TSharedPtr< SWidget >& ValueWidget, float LeftPadding, EHorizontalAlignment HAlign, EVerticalAlignment VAlign, const TArray<FDataprepPropertyLink>& PropertyChain);
+	TSharedRef< SWidget > CreateDefaultWidget(TSharedPtr< SWidget >& NameWidget, TSharedPtr< SWidget >& ValueWidget, float LeftPadding, EHorizontalAlignment HAlign, EVerticalAlignment VAlign, const FDataprepParameterizationContext& ParameterizationContext);
 
 	/** Callback used by all splitters in the details view, so that they move in sync */
 	void OnLeftColumnResized(float InNewWidth)
@@ -130,7 +132,9 @@ private:
 
 	/** Callback used to detect the existence of a new object to display after a reinstancing process */
 	void OnObjectReplaced(const TMap<UObject*, UObject*>& ReplacementObjectMap);
-		
+
+	void ForceRefresh();
+
 private:
 	/** Row generator applied on detailed object */
 	TSharedPtr< class IPropertyRowGenerator > Generator;
@@ -150,13 +154,14 @@ private:
 	/** Relative width to control splitters. */
 	float ColumnWidth;
 
-	/** Points to the currently used colum size data. Can be provided via argument as well.
-	 */
+	// Points to the currently used colum size data. Can be provided via argument as well.
 	TSharedPtr< FDataprepDetailsViewColumnSizeData > ColumnSizeData;
 
 	// If there is a new object to display on the next tick
-	uint8 bNewObjectToDisplay : 1;
+	uint8 bRefreshObjectToDisplay : 1;
 
 	// This pointer to a dataprep asset is used by the parameterization system. It should be nullptr when the parameterization shouldn't be shown
 	TWeakObjectPtr<UDataprepAsset> DataprepAssetForParameterization;
+
+	FDelegateHandle OnDataprepParameterizationWasModifiedHandle;
 };
