Merging CLs #9470941 and #9470915 from //UE4/Dev-Enterprise

9470915:
Rework asset preview widget to become more consistent with the Level Editor's content browser

9470941:
Add the ability to hide/show actors to the scene outliner in the Dataprep Editor

diff --git a/Engine/Plugins/Experimental/DataPrepEditor/Source/DataPrepEditor/Private/DataPrepEditor.cpp b/Engine/Plugins/Experimental/DataPrepEditor/Source/DataPrepEditor/Private/DataPrepEditor.cpp
--- a/Engine/Plugins/Experimental/DataPrepEditor/Source/DataPrepEditor/Private/DataPrepEditor.cpp
+++ b/Engine/Plugins/Experimental/DataPrepEditor/Source/DataPrepEditor/Private/DataPrepEditor.cpp
@@ -48,8 +48,6 @@
 #include "Misc/ScopedSlowTask.h"
 #include "Modules/ModuleManager.h"
 #include "ObjectTools.h"
-#include "SceneOutlinerModule.h"
-#include "SceneOutlinerVisitorTypes.h"
 #include "ScopedTransaction.h"
 #include "Templates/UnrealTemplate.h"
 #include "Toolkits/IToolkit.h"
@@ -89,94 +87,6 @@ const FName FDataprepEditor::ParameterizationDefaultTabId(TEXT("DataprepEditor_P
 
 static bool bLogTiming = true;
 
-namespace DataprepEditorUtils
-{
-	/** Specifics for the scene outliner */
-	using namespace SceneOutliner;
-
-	/**
-	 * This struct is used to force the scene outliner to refuse any rename request
-	 */
-	struct FCanRenameItem : public TTreeItemGetter<bool>
-	{
-		virtual bool Get(const SceneOutliner::FActorTreeItem& ActorItem) const override { return false; };
-		virtual bool Get(const SceneOutliner::FWorldTreeItem& WorldItem) const override { return false; }
-		virtual bool Get(const SceneOutliner::FFolderTreeItem& FolderItem) const override { return false; }
-		virtual bool Get(const SceneOutliner::FComponentTreeItem& ComponentItem) const override { return false; }
-		virtual bool Get(const SceneOutliner::FSubComponentTreeItem& SubComponentItem) const override { return false; }
-	};
-
-	/**
-	 * Use this struct to match the scene outliers selection to a dataprep editor selection
-	 */
-	struct FSynchroniseSelectionToSceneOutliner : public TTreeItemGetter<bool>
-	{
-		FSynchroniseSelectionToSceneOutliner(TSharedRef<FDataprepEditor> InDataprepEditor)
-			: DataprepEditorPtr( InDataprepEditor )
-		{};
-
-		virtual bool Get(const FActorTreeItem& ActorItem) const override
-		{
-			if ( const FDataprepEditor* DataprepEditor = DataprepEditorPtr.Pin().Get() )
-			{
-				return DataprepEditor->GetWorldItemsSelection().Contains( ActorItem.Actor );
-			}
-			return false;
-		}
-
-		virtual bool Get(const FWorldTreeItem& WorldItem) const override
-		{
-			return false;
-		}
-		virtual bool Get(const FFolderTreeItem& FolderItem) const override
-		{ 
-			return false;
-		}
-		virtual bool Get(const FComponentTreeItem& ComponentItem) const override
-		{
-			if ( const FDataprepEditor* DataprepEditor = DataprepEditorPtr.Pin().Get() )
-			{
-				return DataprepEditor->GetWorldItemsSelection().Contains( ComponentItem.Component );
-			}
-			return false;
-		}
-		virtual bool Get(const FSubComponentTreeItem& SubComponentItem) const override
-		{
-			// return this for now as it seams that subcomponent Item is broken or doesn't do what want
-			return false;
-		}
-
-	private:
-		TWeakPtr<FDataprepEditor> DataprepEditorPtr;
-	};
-
-
-	/**
-	 * Use this struct to get the selection from the scene outliner
-	 */
-	struct FGetSelectionFromSceneOutliner : public ITreeItemVisitor
-	{
-		mutable TSet<TWeakObjectPtr<UObject>> Selection;
-
-		virtual void Visit(const FActorTreeItem& ActorItem) const override
-		{
-			Selection.Add( ActorItem.Actor );
-		}
-
-		virtual void Visit(const FWorldTreeItem& WorldItem) const override {}
-		virtual void Visit(const FFolderTreeItem& FolderItem) const override {}
-		virtual void Visit(const FComponentTreeItem& ComponentItem) const override
-		{
-			this->Selection.Add( ComponentItem.Component );
-		}
-
-		virtual void Visit(const FSubComponentTreeItem& SubComponentItem) const override {}
-	};
-
-	/** End of specifics for the scene outliner */
-
-}
-
 class FTimeLogger
 {
 public:
@@ -208,7 +118,6 @@ private:
 	FString Text;
 };
 
-
 namespace DataprepEditorUtil
 {
 	void DeleteActor(AActor* Actor, UWorld* World)
@@ -964,39 +873,6 @@ TSharedRef<SDockTab> FDataprepEditor::SpawnTabPipelineGraph(const FSpawnTabArgs
 }
 // end of temp code for nodes development
 
-void FDataprepEditor::CreateScenePreviewTab()
-{
-	FSceneOutlinerModule& SceneOutlinerModule = FModuleManager::Get().LoadModuleChecked<FSceneOutlinerModule>("SceneOutliner");
-
-	SceneOutliner::FInitializationOptions SceneOutlinerOptions;
-	SceneOutlinerOptions.SpecifiedWorldToDisplay = PreviewWorld.Get();
-
-	SceneOutliner = SceneOutlinerModule.CreateCustomSceneOutliner( SceneOutlinerOptions );
-
-	SceneOutliner->SetSelectionMode( ESelectionMode::Multi )
-		.SetCanRenameItem( MakeUnique<DataprepEditorUtils::FCanRenameItem>() )
-		.SetShouldSelectItemWhenAdded( MakeUnique<DataprepEditorUtils::FSynchroniseSelectionToSceneOutliner>( StaticCastSharedRef<FDataprepEditor>( AsShared() ) ) )
-		.SetShowActorComponents( false )
-		.SetShownOnlySelected( false )
-		.SetShowOnlyCurrentLevel( false )
-		.SetHideTemporaryActors( false );
-
-	SceneOutliner->GetOnItemSelectionChanged().AddSP( this, &FDataprepEditor::OnSceneOutlinerSelectionChanged );
-
-	SAssignNew(ScenePreviewView, SBorder)
-	.Padding(2.f)
-	.BorderImage(FEditorStyle::GetBrush("ToolPanel.GroupBorder"))
-	[
-		SNew(SOverlay)
-		+ SOverlay::Slot()
-		.HAlign(HAlign_Fill)
-		.VAlign(VAlign_Fill)
-		[
-			SceneOutliner.ToSharedRef()
-		]
-	];
-}
-
 TSharedRef<SDockTab> FDataprepEditor::SpawnTabScenePreview(const FSpawnTabArgs & Args)
 {
 	check(Args.GetTabId() == ScenePreviewTabId);
@@ -1293,65 +1169,11 @@ bool FDataprepEditor::CanCommitWorld()
 	return bWorldBuilt && DataprepAssetInterfacePtr->GetConsumer() != nullptr;
 }
 
-
 FString FDataprepEditor::GetTransientContentFolder()
 {
 	return FPaths::Combine( GetRootPackagePath(), FString::FromInt( FPlatformProcess::GetCurrentProcessId() ), SessionID );
 }
 
-void FDataprepEditor::OnSceneOutlinerSelectionChanged(SceneOutliner::FTreeItemPtr ItemPtr, ESelectInfo::Type SelectionMode)
-{
-	using namespace SceneOutliner;
-
-	DataprepEditorUtils::FGetSelectionFromSceneOutliner Visitor;
-
-	for ( FTreeItemPtr Item : SceneOutliner->GetTree().GetSelectedItems() )
-	{
-		Item->Visit( Visitor );
-	}
-
-	SetWorldObjectsSelection( MoveTemp(Visitor.Selection), EWorldSelectionFrom::SceneOutliner );
-}
-
-void FDataprepEditor::SetWorldObjectsSelection(TSet<TWeakObjectPtr<UObject>>&& NewSelection, EWorldSelectionFrom SelectionFrom /* = EWorldSelectionFrom::Unknow */)
-{
-	WorldItemsSelection.Empty( NewSelection.Num() );
-	WorldItemsSelection.Append( MoveTemp(NewSelection) );
-
-	if ( SelectionFrom != EWorldSelectionFrom::SceneOutliner )
-	{
-		DataprepEditorUtils::FSynchroniseSelectionToSceneOutliner Selector( StaticCastSharedRef<FDataprepEditor>( AsShared() ) );
-		SceneOutliner->SetSelection( Selector );
-	}
-
-	if ( SelectionFrom != EWorldSelectionFrom::Viewport )
-	{
-		TArray<AActor*> Actors;
-		Actors.Reserve( WorldItemsSelection.Num() );
-
-		for ( TWeakObjectPtr<UObject> ObjectPtr : WorldItemsSelection )
-		{
-			if ( AActor* Actor = Cast<AActor>( ObjectPtr.Get() ) )
-			{
-				Actors.Add( Actor );
-			}
-		}
-
-		SceneViewportView->SelectActors( Actors );
-	}
-
-	{
-		TSet<UObject*> Objects;
-		Objects.Reserve( WorldItemsSelection.Num() );
-		for ( const TWeakObjectPtr<UObject>& ObjectPtr : WorldItemsSelection )
-		{
-			Objects.Add( ObjectPtr.Get() );
-		}
-
-		SetDetailsObjects( Objects, false );
-	}
-}
-
 bool FDataprepEditor::OnCanExecuteNextStep(UDataprepActionAsset* ActionAsset, UDataprepOperation* Operation, UDataprepFilter* Filter)
 {
 	// #ueent_todo: Make this action configurable by user
diff --git a/Engine/Plugins/Experimental/DataPrepEditor/Source/DataPrepEditor/Private/DataPrepEditorOutliner.cpp b/Engine/Plugins/Experimental/DataPrepEditor/Source/DataPrepEditor/Private/DataPrepEditorOutliner.cpp
new file mode 100644
--- /dev/null
+++ b/Engine/Plugins/Experimental/DataPrepEditor/Source/DataPrepEditor/Private/DataPrepEditorOutliner.cpp
@@ -0,0 +1,616 @@
+// Copyright 1998-2019 Epic Games, Inc. All Rights Reserved.
+
+#include "DataPrepEditor.h"
+
+#include "ICustomSceneOutliner.h"
+#include "ISceneOutlinerColumn.h"
+#include "SceneOutlinerModule.h"
+
+#include "Modules/ModuleManager.h"
+#include "SceneOutlinerModule.h"
+#include "ScopedTransaction.h"
+#include "Widgets/SDataprepEditorViewport.h"
+#include "Widgets/Images/SImage.h"
+#include "Widgets/Layout/SSpacer.h"
+#include "Widgets/Views/STreeView.h"
+
+#define LOCTEXT_NAMESPACE "DataprepEditor"
+
+namespace DataprepEditorSceneOutlinerUtils
+{
+	/** Specifics for the scene outliner */
+	using namespace SceneOutliner;
+
+	/**
+	 * This struct is used to force the scene outliner to refuse any rename request
+	 */
+	struct FCanRenameItem : public TTreeItemGetter<bool>
+	{
+		virtual bool Get(const SceneOutliner::FActorTreeItem& ActorItem) const override { return false; };
+		virtual bool Get(const SceneOutliner::FWorldTreeItem& WorldItem) const override { return false; }
+		virtual bool Get(const SceneOutliner::FFolderTreeItem& FolderItem) const override { return false; }
+		virtual bool Get(const SceneOutliner::FComponentTreeItem& ComponentItem) const override { return false; }
+		virtual bool Get(const SceneOutliner::FSubComponentTreeItem& SubComponentItem) const override { return false; }
+	};
+
+	/**
+	 * Use this struct to match the scene outliers selection to a dataprep editor selection
+	 */
+	struct FSynchroniseSelectionToSceneOutliner : public TTreeItemGetter<bool>
+	{
+		FSynchroniseSelectionToSceneOutliner(TSharedRef<FDataprepEditor> InDataprepEditor)
+			: DataprepEditorPtr(InDataprepEditor)
+		{
+		};
+
+		virtual bool Get(const FActorTreeItem& ActorItem) const override
+		{
+			if (const FDataprepEditor* DataprepEditor = DataprepEditorPtr.Pin().Get())
+			{
+				return DataprepEditor->GetWorldItemsSelection().Contains(ActorItem.Actor);
+			}
+			return false;
+		}
+
+		virtual bool Get(const FWorldTreeItem& WorldItem) const override
+		{
+			return false;
+		}
+		virtual bool Get(const FFolderTreeItem& FolderItem) const override
+		{
+			return false;
+		}
+		virtual bool Get(const FComponentTreeItem& ComponentItem) const override
+		{
+			if (const FDataprepEditor* DataprepEditor = DataprepEditorPtr.Pin().Get())
+			{
+				return DataprepEditor->GetWorldItemsSelection().Contains(ComponentItem.Component);
+			}
+			return false;
+		}
+		virtual bool Get(const FSubComponentTreeItem& SubComponentItem) const override
+		{
+			// return this for now as it seams that subcomponent Item is broken or doesn't do what want
+			return false;
+		}
+
+	private:
+		TWeakPtr<FDataprepEditor> DataprepEditorPtr;
+	};
+
+
+	/**
+	 * Use this struct to get the selection from the scene outliner
+	 */
+	struct FGetSelectionFromSceneOutliner : public ITreeItemVisitor
+	{
+		mutable TSet<TWeakObjectPtr<UObject>> Selection;
+
+		virtual void Visit(const FActorTreeItem& ActorItem) const override
+		{
+			Selection.Add(ActorItem.Actor);
+		}
+
+		virtual void Visit(const FWorldTreeItem& WorldItem) const override {}
+		virtual void Visit(const FFolderTreeItem& FolderItem) const override {}
+		virtual void Visit(const FComponentTreeItem& ComponentItem) const override
+		{
+			this->Selection.Add(ComponentItem.Component);
+		}
+
+		virtual void Visit(const FSubComponentTreeItem& SubComponentItem) const override {}
+	};
+
+	struct FGetVisibilityVisitor : TTreeItemGetter<bool>
+	{
+		/** Map of tree item to visibility */
+		mutable TMap<const ITreeItem*, bool> VisibilityInfo;
+
+		bool RecurseChildren(const ITreeItem& Item) const
+		{
+			if (const bool* Info = VisibilityInfo.Find(&Item))
+			{
+				return *Info;
+			}
+			else
+			{
+				bool bIsVisible = false;
+				for (const auto& ChildPtr : Item.GetChildren())
+				{
+					auto Child = ChildPtr.Pin();
+					if (Child.IsValid() && Child->Get(*this))
+					{
+						bIsVisible = true;
+						break;
+					}
+				}
+				VisibilityInfo.Add(&Item, bIsVisible);
+
+				return bIsVisible;
+			}
+		}
+
+		bool Get(const FActorTreeItem& ActorItem) const
+		{
+			if (const bool* Info = VisibilityInfo.Find(&ActorItem))
+			{
+				return *Info;
+			}
+			else
+			{
+				const AActor* Actor = ActorItem.Actor.Get();
+
+				const bool bIsVisible = Actor && !Actor->IsTemporarilyHiddenInEditor(true);
+				VisibilityInfo.Add(&ActorItem, bIsVisible);
+
+				return bIsVisible;
+			}
+		}
+
+		bool Get(const FWorldTreeItem& WorldItem) const
+		{
+			return RecurseChildren(WorldItem);
+		}
+
+		bool Get(const FFolderTreeItem& FolderItem) const
+		{
+			return RecurseChildren(FolderItem);
+		}
+	};
+
+	struct FSetVisibilityVisitor : IMutableTreeItemVisitor
+	{
+		/** Whether this item should be visible or not */
+		const bool bSetVisibility;
+		TWeakPtr<SDataprepEditorViewport> Viewport;
+
+		FSetVisibilityVisitor(bool bInSetVisibility, TWeakPtr<SDataprepEditorViewport> InViewport)
+			: bSetVisibility(bInSetVisibility)
+			, Viewport(InViewport)
+		{
+		}
+
+		virtual void Visit(FActorTreeItem& ActorItem) const override
+		{
+			AActor* Actor = ActorItem.Actor.Get();
+			if (Actor)
+			{
+				// Save the actor to the transaction buffer to support undo/redo, but do
+				// not call Modify, as we do not want to dirty the actor's package and
+				// we're only editing temporary, transient values
+				SaveToTransactionBuffer(Actor, false);
+				Actor->SetIsTemporarilyHiddenInEditor(!bSetVisibility);
+
+				Viewport.Pin()->SetActorVisibility(Actor, bSetVisibility);
+
+				// Apply the same visibility to the actors children
+				for (auto& ChildPtr : ActorItem.GetChildren())
+				{
+					auto Child = ChildPtr.Pin();
+					if (Child.IsValid())
+					{
+						FSetVisibilityVisitor Visibility(bSetVisibility, Viewport);
+						Child->Visit(Visibility);
+					}
+				}
+			}
+		}
+
+		virtual void Visit(FWorldTreeItem& WorldItem) const override
+		{
+			for (auto& ChildPtr : WorldItem.GetChildren())
+			{
+				auto Child = ChildPtr.Pin();
+				if (Child.IsValid())
+				{
+					FSetVisibilityVisitor Visibility(bSetVisibility, Viewport);
+					Child->Visit(Visibility);
+				}
+			}
+		}
+
+		virtual void Visit(FFolderTreeItem& FolderItem) const override
+		{
+			for (auto& ChildPtr : FolderItem.GetChildren())
+			{
+				auto Child = ChildPtr.Pin();
+				if (Child.IsValid())
+				{
+					FSetVisibilityVisitor Visibility(bSetVisibility, Viewport);
+					Child->Visit(Visibility);
+				}
+			}
+		}
+	};
+
+	class FVisibilityDragDropOp : public FDragDropOperation, public TSharedFromThis<FVisibilityDragDropOp>
+	{
+	public:
+
+		DRAG_DROP_OPERATOR_TYPE(FVisibilityDragDropOp, FDragDropOperation)
+
+		/** Flag which defines whether to hide destination actors or not */
+		bool bHidden;
+
+		/** Undo transaction stolen from the gutter which is kept alive for the duration of the drag */
+		TUniquePtr<FScopedTransaction> UndoTransaction;
+
+		/** The widget decorator to use */
+		virtual TSharedPtr<SWidget> GetDefaultDecorator() const override
+		{
+			return SNullWidget::NullWidget;
+		}
+
+		/** Create a new drag and drop operation out of the specified flag */
+		static TSharedRef<FVisibilityDragDropOp> New(const bool _bHidden, TUniquePtr<FScopedTransaction>& ScopedTransaction)
+		{
+			TSharedRef<FVisibilityDragDropOp> Operation = MakeShareable(new FVisibilityDragDropOp);
+
+			Operation->bHidden = _bHidden;
+			Operation->UndoTransaction = MoveTemp(ScopedTransaction);
+
+			Operation->Construct();
+			return Operation;
+		}
+	};
+
+	class FPreviewSceneOutlinerGutter : public ISceneOutlinerColumn
+	{
+	public:
+		FPreviewSceneOutlinerGutter(ISceneOutliner& Outliner, TWeakPtr<SDataprepEditorViewport> InViewport)
+		{
+			WeakSceneViewport = InViewport;
+			WeakOutliner = StaticCastSharedRef<ISceneOutliner>(Outliner.AsShared());
+		}
+
+		virtual ~FPreviewSceneOutlinerGutter() {}
+
+		static FName GetID() { return FName("PreviewGutter"); }
+
+		virtual void Tick(double InCurrentTime, float InDeltaTime) override
+		{
+			VisibilityCache.VisibilityInfo.Empty();
+		}
+
+		virtual FName GetColumnID() override { return GetID(); }
+
+		virtual SHeaderRow::FColumn::FArguments ConstructHeaderRowColumn() override
+		{
+			return SHeaderRow::Column(GetColumnID())
+				.FixedWidth(16.f)
+				[
+					SNew(SSpacer)
+				];
+		}
+
+		virtual const TSharedRef<SWidget> ConstructRowWidget(FTreeItemRef TreeItem, const STableRow<FTreeItemPtr>& Row) override;
+
+		virtual void SortItems(TArray<FTreeItemPtr>& RootItems, const EColumnSortMode::Type SortMode) const override {}
+
+		/** Check whether the specified item is visible */
+		bool IsItemVisible(const ITreeItem& Item) { return Item.Get(VisibilityCache); }
+
+		TWeakPtr<SDataprepEditorViewport> GetViewport() const
+		{
+			return WeakSceneViewport;
+		}
+
+	private:
+		/** Weak pointer back to the scene outliner - required for setting visibility on current selection. */
+		TWeakPtr<ISceneOutliner> WeakOutliner;
+
+		TWeakPtr<SDataprepEditorViewport> WeakSceneViewport;
+
+		/** Visitor used to get (and cache) visibilty for items. Cahced per-frame to avoid expensive recursion. */
+		FGetVisibilityVisitor VisibilityCache;
+	};
+
+	/** Widget responsible for managing the visibility for a single actor */
+	class SVisibilityWidget : public SImage
+	{
+	public:
+		SLATE_BEGIN_ARGS(SVisibilityWidget) {}
+		SLATE_END_ARGS()
+
+		/** Construct this widget */
+		void Construct(const FArguments& InArgs, TWeakPtr<FPreviewSceneOutlinerGutter> InWeakColumn, TWeakPtr<ISceneOutliner> InWeakOutliner, TWeakPtr<ITreeItem> InWeakTreeItem)
+		{
+			WeakTreeItem = InWeakTreeItem;
+			WeakOutliner = InWeakOutliner;
+			WeakColumn = InWeakColumn;
+
+			SImage::Construct(
+				SImage::FArguments()
+				.Image(this, &SVisibilityWidget::GetBrush)
+			);
+		}
+
+	private:
+
+		/** Start a new drag/drop operation for this widget */
+		virtual FReply OnDragDetected(const FGeometry& MyGeometry, const FPointerEvent& MouseEvent) override
+		{
+			if (MouseEvent.IsMouseButtonDown(EKeys::LeftMouseButton))
+			{
+				return FReply::Handled().BeginDragDrop(FVisibilityDragDropOp::New(!IsVisible(), UndoTransaction));
+			}
+			else
+			{
+				return FReply::Unhandled();
+			}
+		}
+
+		/** If a visibility drag drop operation has entered this widget, set its actor to the new visibility state */
+		virtual void OnDragEnter(const FGeometry& MyGeometry, const FDragDropEvent& DragDropEvent) override
+		{
+			auto VisibilityOp = DragDropEvent.GetOperationAs<FVisibilityDragDropOp>();
+			if (VisibilityOp.IsValid())
+			{
+				SetIsVisible(!VisibilityOp->bHidden);
+			}
+		}
+
+		FReply HandleClick()
+		{
+			auto Outliner = WeakOutliner.Pin();
+			auto TreeItem = WeakTreeItem.Pin();
+			auto Column = WeakColumn.Pin();
+
+			if (!Outliner.IsValid() || !TreeItem.IsValid() || !Column.IsValid())
+			{
+				return FReply::Unhandled();
+			}
+
+			// Open an undo transaction
+			UndoTransaction.Reset(new FScopedTransaction(LOCTEXT("SetActorVisibility", "Set Actor Visibility")));
+
+			const auto& Tree = Outliner->GetTree();
+
+			const bool bVisible = !IsVisible();
+
+			// We operate on all the selected items if the specified item is selected
+			if (Tree.IsItemSelected(TreeItem.ToSharedRef()))
+			{
+				const FSetVisibilityVisitor Visitor(bVisible, WeakColumn.Pin()->GetViewport());
+
+				for (auto& SelectedItem : Tree.GetSelectedItems())
+				{
+					if (IsVisible(SelectedItem, Column) != bVisible)
+					{
+						SelectedItem->Visit(Visitor);
+					}
+				}
+
+				GEditor->RedrawAllViewports();
+			}
+			else
+			{
+				SetIsVisible(bVisible);
+			}
+
+			return FReply::Handled().DetectDrag(SharedThis(this), EKeys::LeftMouseButton);
+		}
+
+		virtual FReply OnMouseButtonDoubleClick(const FGeometry& InMyGeometry, const FPointerEvent& InMouseEvent) override
+		{
+			return HandleClick();
+		}
+
+		/** Called when the mouse button is pressed down on this widget */
+		virtual FReply OnMouseButtonDown(const FGeometry& MyGeometry, const FPointerEvent& MouseEvent) override
+		{
+			if (MouseEvent.GetEffectingButton() != EKeys::LeftMouseButton)
+			{
+				return FReply::Unhandled();
+			}
+
+			return HandleClick();
+		}
+
+		/** Process a mouse up message */
+		virtual FReply OnMouseButtonUp(const FGeometry& MyGeometry, const FPointerEvent& MouseEvent) override
+		{
+			if (MouseEvent.GetEffectingButton() == EKeys::LeftMouseButton)
+			{
+				UndoTransaction.Reset();
+				return FReply::Handled();
+			}
+
+			return FReply::Unhandled();
+		}
+
+		/** Called when this widget had captured the mouse, but that capture has been revoked for some reason. */
+		virtual void OnMouseCaptureLost(const FCaptureLostEvent& CaptureLostEvent) override
+		{
+			UndoTransaction.Reset();
+		}
+
+		/** Get the brush for this widget */
+		const FSlateBrush* GetBrush() const
+		{
+			if (IsVisible())
+			{
+				static const FName NAME_VisibleHoveredBrush = TEXT("Level.VisibleHighlightIcon16x");
+				static const FName NAME_VisibleNotHoveredBrush = TEXT("Level.VisibleIcon16x");
+				return IsHovered() ? FEditorStyle::GetBrush(NAME_VisibleHoveredBrush) :
+					FEditorStyle::GetBrush(NAME_VisibleNotHoveredBrush);
+			}
+			else
+			{
+				static const FName NAME_NotVisibleHoveredBrush = TEXT("Level.NotVisibleHighlightIcon16x");
+				static const FName NAME_NotVisibleNotHoveredBrush = TEXT("Level.NotVisibleIcon16x");
+				return IsHovered() ? FEditorStyle::GetBrush(NAME_NotVisibleHoveredBrush) :
+					FEditorStyle::GetBrush(NAME_NotVisibleNotHoveredBrush);
+			}
+		}
+
+		/** Check if the specified item is visible */
+		static bool IsVisible(const FTreeItemPtr& Item, const TSharedPtr<FPreviewSceneOutlinerGutter>& Column)
+		{
+			return Column.IsValid() && Item.IsValid() ? Column->IsItemVisible(*Item) : false;
+		}
+
+		/** Check if our wrapped tree item is visible */
+		bool IsVisible() const
+		{
+			return IsVisible(WeakTreeItem.Pin(), WeakColumn.Pin());
+		}
+
+		/** Set the actor this widget is responsible for to be hidden or shown */
+		void SetIsVisible(const bool bVisible)
+		{
+			TSharedPtr<ITreeItem> TreeItem = WeakTreeItem.Pin();
+			TSharedPtr<ISceneOutliner> Outliner = WeakOutliner.Pin();
+
+			if (TreeItem.IsValid() && Outliner.IsValid() && IsVisible() != bVisible)
+			{
+				FSetVisibilityVisitor Visitor(bVisible, WeakColumn.Pin()->GetViewport());
+				TreeItem->Visit(Visitor);
+
+				Outliner->Refresh();
+
+				GEditor->RedrawAllViewports();
+			}
+		}
+
+		/** The tree item we relate to */
+		TWeakPtr<ITreeItem> WeakTreeItem;
+
+		/** Reference back to the outliner so we can set visibility of a whole selection */
+		TWeakPtr<ISceneOutliner> WeakOutliner;
+
+		/** Weak pointer back to the column */
+		TWeakPtr<FPreviewSceneOutlinerGutter> WeakColumn;
+
+		/** Scoped undo transaction */
+		TUniquePtr<FScopedTransaction> UndoTransaction;
+	};
+
+	const TSharedRef<SWidget> FPreviewSceneOutlinerGutter::ConstructRowWidget(FTreeItemRef TreeItem, const STableRow<FTreeItemPtr>& Row)
+	{
+		return SNew(SHorizontalBox)
+			+ SHorizontalBox::Slot()
+			.AutoWidth()
+			.VAlign(VAlign_Center)
+			[
+				SNew(SVisibilityWidget, SharedThis(this), WeakOutliner, TreeItem)
+			];
+	}
+
+	/** End of specifics for the scene outliner */
+}
+
+void FDataprepEditor::CreateScenePreviewTab()
+{
+	FSceneOutlinerModule& SceneOutlinerModule = FModuleManager::Get().LoadModuleChecked<FSceneOutlinerModule>("SceneOutliner");
+
+	SceneOutliner::FInitializationOptions SceneOutlinerOptions;
+	SceneOutlinerOptions.SpecifiedWorldToDisplay = PreviewWorld.Get();
+
+	SceneOutliner = SceneOutlinerModule.CreateCustomSceneOutliner(SceneOutlinerOptions);
+
+	// Add our custom visibility gutter
+
+	auto CreateColumn = [this](ISceneOutliner& Outliner)
+	{
+		return TSharedRef< ISceneOutlinerColumn >(MakeShareable(new DataprepEditorSceneOutlinerUtils::FPreviewSceneOutlinerGutter(Outliner, SceneViewportView.ToSharedRef())));
+	};
+
+	SceneOutliner::FColumnInfo ColumnInfo;
+	ColumnInfo.Visibility = SceneOutliner::EColumnVisibility::Visible;
+	ColumnInfo.PriorityIndex = 0;
+	ColumnInfo.Factory.BindLambda([&](ISceneOutliner& Outliner)
+	{
+		return TSharedRef< ISceneOutlinerColumn >(MakeShareable(new DataprepEditorSceneOutlinerUtils::FPreviewSceneOutlinerGutter(Outliner, SceneViewportView.ToSharedRef())));
+	});
+
+	SceneOutliner->AddColumn(DataprepEditorSceneOutlinerUtils::FPreviewSceneOutlinerGutter::GetID(), ColumnInfo);
+
+	// Add the default outliner columns
+	const SceneOutliner::FSharedOutlinerData& SharedData = SceneOutliner->GetSharedData();
+	for (auto& DefaultColumn : SceneOutlinerModule.DefaultColumnMap)
+	{
+		if (!DefaultColumn.Value.ValidMode.IsSet() || SharedData.Mode == DefaultColumn.Value.ValidMode.GetValue())
+		{
+			SceneOutliner->AddColumn(DefaultColumn.Key, DefaultColumn.Value.ColumnInfo);
+		}
+	}
+
+	SceneOutliner->SetSelectionMode(ESelectionMode::Multi)
+		.SetCanRenameItem(MakeUnique<DataprepEditorSceneOutlinerUtils::FCanRenameItem>())
+		.SetShouldSelectItemWhenAdded(MakeUnique<DataprepEditorSceneOutlinerUtils::FSynchroniseSelectionToSceneOutliner>(StaticCastSharedRef<FDataprepEditor>(AsShared())))
+		.SetShowActorComponents(false)
+		.SetShownOnlySelected(false)
+		.SetShowOnlyCurrentLevel(false)
+		.SetHideTemporaryActors(false);
+
+	SceneOutliner->GetOnItemSelectionChanged().AddSP(this, &FDataprepEditor::OnSceneOutlinerSelectionChanged);
+
+	SAssignNew(ScenePreviewView, SBorder)
+		.Padding(2.f)
+		.BorderImage(FEditorStyle::GetBrush("ToolPanel.GroupBorder"))
+		[
+			SNew(SOverlay)
+			+ SOverlay::Slot()
+			.HAlign(HAlign_Fill)
+			.VAlign(VAlign_Fill)
+			[
+				SceneOutliner.ToSharedRef()
+			]
+		];
+}
+
+void FDataprepEditor::OnSceneOutlinerSelectionChanged(SceneOutliner::FTreeItemPtr ItemPtr, ESelectInfo::Type SelectionMode)
+{
+	using namespace SceneOutliner;
+
+	DataprepEditorSceneOutlinerUtils::FGetSelectionFromSceneOutliner Visitor;
+
+	for (FTreeItemPtr Item : SceneOutliner->GetTree().GetSelectedItems())
+	{
+		Item->Visit(Visitor);
+	}
+
+	SetWorldObjectsSelection(MoveTemp(Visitor.Selection), EWorldSelectionFrom::SceneOutliner);
+}
+
+void FDataprepEditor::SetWorldObjectsSelection(TSet<TWeakObjectPtr<UObject>>&& NewSelection, EWorldSelectionFrom SelectionFrom /* = EWorldSelectionFrom::Unknow */)
+{
+	WorldItemsSelection.Empty(NewSelection.Num());
+	WorldItemsSelection.Append(MoveTemp(NewSelection));
+
+	if (SelectionFrom != EWorldSelectionFrom::SceneOutliner)
+	{
+		DataprepEditorSceneOutlinerUtils::FSynchroniseSelectionToSceneOutliner Selector(StaticCastSharedRef<FDataprepEditor>(AsShared()));
+		SceneOutliner->SetSelection(Selector);
+	}
+
+	if (SelectionFrom != EWorldSelectionFrom::Viewport)
+	{
+		TArray<AActor*> Actors;
+		Actors.Reserve(WorldItemsSelection.Num());
+
+		for (TWeakObjectPtr<UObject> ObjectPtr : WorldItemsSelection)
+		{
+			if (AActor* Actor = Cast<AActor>(ObjectPtr.Get()))
+			{
+				Actors.Add(Actor);
+			}
+		}
+
+		SceneViewportView->SelectActors(Actors);
+	}
+
+	{
+		TSet<UObject*> Objects;
+		Objects.Reserve(WorldItemsSelection.Num());
+		for (const TWeakObjectPtr<UObject>& ObjectPtr : WorldItemsSelection)
+		{
+			Objects.Add(ObjectPtr.Get());
+		}
+
+		SetDetailsObjects(Objects, false);
+	}
+}
+
+#undef LOCTEXT_NAMESPACE
diff --git a/Engine/Plugins/Experimental/DataPrepEditor/Source/DataPrepEditor/Private/Widgets/SAssetsPreviewWidget.cpp b/Engine/Plugins/Experimental/DataPrepEditor/Source/DataPrepEditor/Private/Widgets/SAssetsPreviewWidget.cpp
--- a/Engine/Plugins/Experimental/DataPrepEditor/Source/DataPrepEditor/Private/Widgets/SAssetsPreviewWidget.cpp
+++ b/Engine/Plugins/Experimental/DataPrepEditor/Source/DataPrepEditor/Private/Widgets/SAssetsPreviewWidget.cpp
@@ -3,6 +3,7 @@
 #include "SAssetsPreviewWidget.h"
 
 #include "EditorStyleSet.h"
+#include "Widgets/Images/SImage.h"
 #include "Widgets/Text/STextBlock.h"
 #include "SAssetSearchBox.h"
 #include "Misc/PackageName.h"
@@ -103,6 +104,78 @@ namespace AssetPreviewWidget
 		}
 	};
 
+	/** Represents a row in the AssetPreview's tree view */
+	class SAssetPreviewTableRow : public STableRow<FAssetTreeItemPtr>
+	{
+	public:
+		SLATE_BEGIN_ARGS(SAssetPreviewTableRow) {}
+		SLATE_END_ARGS()
+
+		void Construct(const FArguments& InArgs, const TSharedRef<STableViewBase>& OwnerTableView, FAssetTreeItemPtr InItem, TSharedRef< const SAssetsPreviewWidget > InPreviewWidget)
+		{
+			FolderOpenBrush = FEditorStyle::GetBrush("ContentBrowser.AssetTreeFolderOpen");
+			FolderClosedBrush = FEditorStyle::GetBrush("ContentBrowser.AssetTreeFolderClosed");
+			AssetIconBrush = FEditorStyle::GetBrush("ContentBrowser.ColumnViewAssetIcon");
+
+			PreviewWidgetWeakPtr = InPreviewWidget;
+			ItemWeakPtr = InItem;
+
+			STableRow::Construct(
+				STableRow::FArguments()
+				.Style(FEditorStyle::Get(), "ContentBrowser.AssetListView.TableRow")
+				.Cursor(EMouseCursor::Default),
+				OwnerTableView);
+
+
+			TSharedRef< SHorizontalBox > Widget = SNew(SHorizontalBox)
+				+ SHorizontalBox::Slot()
+				.AutoWidth()
+				.Padding(0, 0, 2, 0)
+				.VAlign(VAlign_Center)
+				[
+					// Item icon
+					SNew(SImage)
+					.Image(this, &SAssetPreviewTableRow::GetIconBrush)
+					.ColorAndOpacity(FSlateColor(FLinearColor::White))
+				]
+
+				+ SHorizontalBox::Slot()
+				.AutoWidth()
+				.VAlign(VAlign_Center)
+				[
+					SNew(STextBlock)
+					.Text(FText::FromString(*(InItem->Name)))
+					.Font(FEditorStyle::GetFontStyle("ContentBrowser.SourceTreeItemFont"))
+					.HighlightText(PreviewWidgetWeakPtr.Pin().Get(), &SAssetsPreviewWidget::OnGetHighlightText)
+				];
+
+			SetContent(Widget);
+		}
+
+	private:
+		const FSlateBrush* GetIconBrush() const
+		{
+			const FSlateBrush* IconBrush = AssetIconBrush;
+			if (ItemWeakPtr.Pin()->IsFolder())
+			{
+				const bool bExpanded = PreviewWidgetWeakPtr.Pin()->GetTreeView()->IsItemExpanded(ItemWeakPtr.Pin());
+				IconBrush = bExpanded ? FolderOpenBrush : FolderClosedBrush;
+			}
+			return IconBrush;
+		}
+
+	private:
+		/** Brushes for the different folder states */
+		const FSlateBrush* FolderOpenBrush;
+		const FSlateBrush* FolderClosedBrush;
+		const FSlateBrush* AssetIconBrush;
+
+		FAssetTreeItemWeakPtr ItemWeakPtr;
+
+		/** Weak reference back to the preview widget that owns us */
+		TWeakPtr< const SAssetsPreviewWidget > PreviewWidgetWeakPtr;
+	};
+
 	void SAssetsPreviewWidget::Construct(const FArguments& InArgs)
 	{
 		ChildSlot
@@ -150,6 +223,17 @@ namespace AssetPreviewWidget
 		RootItems.Empty();
 		TMap< FString, FAssetTreeItemPtr > NamesToRootItem;
 
+		// Make sure the root dir is displayed as "Content".
+		// This is more descriptive for the end user.
+		TArray<FString> Tokens;
+		const FString StrContent = TEXT("Content");
+		SubstitutePath.ParseIntoArray(Tokens, TEXT("/"));
+		if (Tokens.Num() > 0 && !Tokens[0].Equals(StrContent))
+		{
+			const int32 StartChar = SubstitutePath.Find(Tokens[0]);
+			SubstitutePath.RemoveAt(StartChar, Tokens[0].Len(), false);
+			SubstitutePath.InsertAt(StartChar, StrContent);
+		}
 
 		for (const TWeakObjectPtr< UObject >& Asset : InAssetsList)
 		{
@@ -218,6 +302,9 @@ namespace AssetPreviewWidget
 			}
 		}
 
+		// Sort items in alphabetical order
+		SortRecursive(RootItems);
+
 		FilterAssetsNames();
 	}
 
@@ -261,6 +348,17 @@ namespace AssetPreviewWidget
 		}
 	}
 
+	void SAssetsPreviewWidget::SortRecursive(TArray< FAssetTreeItemPtr >& InItems)
+	{
+		InItems.Sort([](const FAssetTreeItemPtr A, const FAssetTreeItemPtr B) { return A->Name.Compare(B->Name) < 0; });
+
+		for (int ItemIndex = 0; ItemIndex < InItems.Num(); ++ItemIndex)
+		{
+			SortRecursive(InItems[ItemIndex]->Folders);
+			SortRecursive(InItems[ItemIndex]->Assets);
+		}
+	}
+
 	TArray< FString > SAssetsPreviewWidget::GetItemsName(TWeakObjectPtr< UObject > Asset) const
 	{
 		TArray< FString > ItemsName;
@@ -276,19 +374,9 @@ namespace AssetPreviewWidget
 
 	TSharedRef< ITableRow > SAssetsPreviewWidget::MakeRowWidget(FAssetTreeItemPtr InItem, const TSharedRef< STableViewBase >& OwnerTable) const
 	{
-		TSharedPtr< STableRow< FAssetTreeItemPtr > > TableRowWidget;
-
-		SAssignNew(TableRowWidget, STableRow< FAssetTreeItemPtr >, OwnerTable)
-			.Style(FEditorStyle::Get(), "ContentBrowser.AssetListView.TableRow")
-			.Cursor(EMouseCursor::Default);
-
-		TSharedRef< STextBlock > Item = SNew(STextBlock)
-			//.TextStyle( FEditorStyle::Get(), "LargeText" )
-			.Text(FText::FromString(*(InItem->Name)))
-			.Font(FEditorStyle::GetFontStyle("ContentBrowser.SourceTreeItemFont"))
-			.HighlightText(this, &SAssetsPreviewWidget::OnGetHighlightText);
+		TSharedPtr< SAssetPreviewTableRow > TableRowWidget;
 
-		TableRowWidget->SetContent(Item);
+		SAssignNew(TableRowWidget, SAssetPreviewTableRow, OwnerTable, InItem, SharedThis(this));
 
 		return TableRowWidget.ToSharedRef();
 	}
diff --git a/Engine/Plugins/Experimental/DataPrepEditor/Source/DataPrepEditor/Private/Widgets/SAssetsPreviewWidget.h b/Engine/Plugins/Experimental/DataPrepEditor/Source/DataPrepEditor/Private/Widgets/SAssetsPreviewWidget.h
--- a/Engine/Plugins/Experimental/DataPrepEditor/Source/DataPrepEditor/Private/Widgets/SAssetsPreviewWidget.h
+++ b/Engine/Plugins/Experimental/DataPrepEditor/Source/DataPrepEditor/Private/Widgets/SAssetsPreviewWidget.h
@@ -32,6 +32,14 @@ namespace AssetPreviewWidget
 		void ClearAssetList();
 
 		FOnSelectionChanged& OnSelectionChanged() { return OnSelectionChangedDelegate; };
+	
+		FText OnGetHighlightText() const;
+
+		const TSharedPtr< STreeView< FAssetTreeItemPtr > > GetTreeView() const
+		{
+			return TreeView;
+		}
+
 	private:
 
 		friend struct FAssetTreeItem;
@@ -40,6 +48,8 @@ namespace AssetPreviewWidget
 		void ExpandAllFolders();
 		void ExpandFolderRecursive(FAssetTreeItemPtr InItem);
 
+		void SortRecursive(TArray< FAssetTreeItemPtr >& InItems);
+
 		TArray< FString > GetItemsName(TWeakObjectPtr< UObject > Asset) const;
 
 		TSharedRef< class ITableRow > MakeRowWidget(FAssetTreeItemPtr InItem, const TSharedRef< class STableViewBase >& OwnerTable) const;
@@ -48,8 +58,6 @@ namespace AssetPreviewWidget
 		void OnSearchBoxChanged(const FText& InSearchText);
 		void OnSearchBoxCommitted(const FText& InSearchText, ETextCommit::Type CommitInfo);
 
-		FText OnGetHighlightText() const;
-
 		void OnSetExpansionRecursive(FAssetTreeItemPtr InTreeNode, bool bInIsItemExpanded);
 
 		void OnSelectionChangedInternal(FAssetTreeItemPtr ItemSelected, ESelectInfo::Type SelectionType);
diff --git a/Engine/Plugins/Experimental/DataPrepEditor/Source/DataPrepEditor/Private/Widgets/SDataprepEditorViewport.cpp b/Engine/Plugins/Experimental/DataPrepEditor/Source/DataPrepEditor/Private/Widgets/SDataprepEditorViewport.cpp
--- a/Engine/Plugins/Experimental/DataPrepEditor/Source/DataPrepEditor/Private/Widgets/SDataprepEditorViewport.cpp
+++ b/Engine/Plugins/Experimental/DataPrepEditor/Source/DataPrepEditor/Private/Widgets/SDataprepEditorViewport.cpp
@@ -1020,6 +1020,20 @@ void SDataprepEditorViewport::SelectActors(const TArray< AActor* >& SelectedActo
 	UpdateSelection();
 }
 
+void SDataprepEditorViewport::SetActorVisibility(AActor* SceneActor, bool bInVisibility)
+{
+	TArray< UStaticMeshComponent* > SceneComponents;
+	SceneActor->GetComponents< UStaticMeshComponent >(SceneComponents);
+	for (UStaticMeshComponent* SceneComponent : SceneComponents)
+	{
+		UStaticMeshComponent** PreviewComponent = MeshComponentsMapping.Find(SceneComponent);
+		if (PreviewComponent)
+		{
+			(*PreviewComponent)->SetVisibility(bInVisibility);
+		}
+	}
+}
+
 void SDataprepEditorViewport::UpdateSelection()
 {
 	TSharedPtr<FDataprepEditor> DataprepEditorPtr = DataprepEditor.Pin();
diff --git a/Engine/Plugins/Experimental/DataPrepEditor/Source/DataPrepEditor/Private/Widgets/SDataprepEditorViewport.h b/Engine/Plugins/Experimental/DataPrepEditor/Source/DataPrepEditor/Private/Widgets/SDataprepEditorViewport.h
--- a/Engine/Plugins/Experimental/DataPrepEditor/Source/DataPrepEditor/Private/Widgets/SDataprepEditorViewport.h
+++ b/Engine/Plugins/Experimental/DataPrepEditor/Source/DataPrepEditor/Private/Widgets/SDataprepEditorViewport.h
@@ -170,6 +170,9 @@ public:
 
 	void SelectActors( const TArray< AActor* >& SelectedActors );
 
+	/** Change visibility of preview mesh components */
+	void SetActorVisibility(AActor* SceneActor, bool bInVisibility);
+
 protected:
 	/** SEditorViewport interface */
 	virtual TSharedRef<FEditorViewportClient> MakeEditorViewportClient() override;
