Merging CL #9489306 and #9491781 from //UE4/Dev-Enterprise

9489306:
Fixed missing reference to static meshes when DnD a Datasmith scene generated by a Dataprep asset

9491781:
Made a fast access container for the dataprep bindings.

diff --git a/Engine/Plugins/Enterprise/DatasmithImporter/Source/DatasmithImporter/Private/DatasmithConsumer.cpp b/Engine/Plugins/Enterprise/DatasmithImporter/Source/DatasmithImporter/Private/DatasmithConsumer.cpp
--- a/Engine/Plugins/Enterprise/DatasmithImporter/Source/DatasmithImporter/Private/DatasmithConsumer.cpp
+++ b/Engine/Plugins/Enterprise/DatasmithImporter/Source/DatasmithImporter/Private/DatasmithConsumer.cpp
@@ -77,12 +77,6 @@ namespace DatasmithConsumerUtils
 
 		return DatasmithUniqueId.IsEmpty() ? Object->GetName() : DatasmithUniqueId;
 	}
-
-	FString GetObjectTag(UObject* Object)
-	{
-		const FString ObjectPath = FPaths::Combine( Object->GetOutermost()->GetName(), Object->GetName() );
-		return FMD5::HashBytes( reinterpret_cast<const uint8*>(*ObjectPath), ObjectPath.Len() * sizeof(TCHAR) );
-	}
 }
 
 bool UDatasmithConsumer::Initialize()
@@ -233,9 +227,6 @@ bool UDatasmithConsumer::BuildContexts( UWorld* ImportWorld )
 	ImportContextPtr->ActorsContext.ImportSceneActor->GetAttachedActors( RootActors );
 	FDatasmithImporterUtils::FillSceneElement( ImportContextPtr->Scene, RootActors );
 
-	// Store IDatasmithScene(Element) in UDatasmithScene
-	FDatasmithImporterUtils::SaveDatasmithScene( ImportContextPtr->Scene.ToSharedRef(), ImportContextPtr->SceneAsset );
-
 	// Initialize context
 	FString SceneOuterPath = DatasmithScene->GetOutermost()->GetName();
 	FString RootPath = FPackageName::GetLongPackagePath( SceneOuterPath );
@@ -285,6 +276,8 @@ bool UDatasmithConsumer::BuildContexts( UWorld* ImportWorld )
 	// Add assets as if they have been imported using the current import context
 	DatasmithConsumerUtils::AddAssetsToContext( *ImportContextPtr, Context.Assets );
 
+	// Store IDatasmithScene(Element) in UDatasmithScene
+	FDatasmithImporterUtils::SaveDatasmithScene( ImportContextPtr->Scene.ToSharedRef(), ImportContextPtr->SceneAsset );
 	return true;
 }
 
@@ -650,17 +643,20 @@ namespace DatasmithConsumerUtils
 		{
 			if( UObject* Asset = AssetPtr.Get() )
 			{
-				FString AssetTag = DatasmithConsumerUtils::GetObjectTag( Asset );
+				FString AssetTag = FDatasmithImporterUtils::GetDatasmithElementIdString( Asset );
 
 				if(UTexture* Texture = Cast<UTexture>(Asset))
 				{
 					TSharedRef< IDatasmithTextureElement > TextureElement = FDatasmithSceneFactory::CreateTexture( *AssetTag );
+					TextureElement->SetLabel( *Texture->GetName() );
+
 					ImportContext.ImportedTextures.Add( TextureElement, Texture );
+					ImportContext.Scene->AddTexture( TextureElement );
 				}
 				else if(UMaterialInstance* MaterialInstance = Cast<UMaterialInstance>(Asset))
 				{
 					TSharedRef< IDatasmithMaterialElement > MaterialElement = FDatasmithSceneFactory::CreateMaterial( *AssetTag );
-					ImportContext.ImportedMaterials.Add( MaterialElement, MaterialInstance );
+					MaterialElement->SetLabel( *MaterialInstance->GetName() );
 
 					if ( UMaterialInterface* MaterialParent = MaterialInstance->Parent )
 					{
@@ -685,6 +681,9 @@ namespace DatasmithConsumerUtils
 							AddTemplate( UDatasmithMaterialInstanceTemplate::StaticClass(), MaterialInstanceConstant );
 						}
 					}
+
+					ImportContext.ImportedMaterials.Add( MaterialElement, MaterialInstance );
+					ImportContext.Scene->AddMaterial( MaterialElement );
 				}
 				else if(UStaticMesh* StaticMesh = Cast<UStaticMesh>(Asset))
 				{
@@ -700,17 +699,33 @@ namespace DatasmithConsumerUtils
 					}
 
 					TSharedRef< IDatasmithMeshElement > MeshElement = FDatasmithSceneFactory::CreateMesh( *AssetTag );
+					MeshElement->SetLabel( *StaticMesh->GetName() );
+
+
+					for(int32 Index = 0; Index < StaticMesh->GetNumSections( 0 ); ++Index)
+					{
+						const FString MaterialTag = FDatasmithImporterUtils::GetDatasmithElementIdString( StaticMesh->GetMaterial( Index ) );
+						MeshElement->SetMaterial( *MaterialTag, Index );
+					}
+
 					ImportContext.ImportedStaticMeshes.Add( MeshElement, StaticMesh );
+					ImportContext.Scene->AddMesh( MeshElement );
 				}
 				else if(ULevelSequence* LevelSequence = Cast<ULevelSequence>(Asset))
 				{
 					TSharedRef< IDatasmithLevelSequenceElement > LevelSequenceElement = FDatasmithSceneFactory::CreateLevelSequence( *AssetTag );
+					LevelSequenceElement->SetLabel( *LevelSequence->GetName() );
+
 					ImportContext.ImportedLevelSequences.Add( LevelSequenceElement, LevelSequence );
+					ImportContext.Scene->AddLevelSequence( LevelSequenceElement );
 				}
 				else if(ULevelVariantSets* LevelVariantSets = Cast<ULevelVariantSets>(Asset))
 				{
 					TSharedRef< IDatasmithLevelVariantSetsElement > LevelVariantSetsElement = FDatasmithSceneFactory::CreateLevelVariantSets( *AssetTag );
+					LevelVariantSetsElement->SetLabel( *LevelVariantSets->GetName() );
+
 					ImportContext.ImportedLevelVariantSets.Add( LevelVariantSetsElement, LevelVariantSets );
+					ImportContext.Scene->AddLevelVariantSets( LevelVariantSetsElement );
 				}
 				// #ueent_todo: Add support for assets which are not of the classes above
 			}
@@ -723,9 +738,12 @@ namespace DatasmithConsumerUtils
 			{
 				if( !ParentMaterials.Contains( Material ) )
 				{
-					FString AssetTag = DatasmithConsumerUtils::GetObjectTag( Material );
+					FString AssetTag = FDatasmithImporterUtils::GetDatasmithElementIdString( Material );
 					TSharedRef< IDatasmithMaterialElement > MaterialElement = FDatasmithSceneFactory::CreateMaterial( *AssetTag );
+					MaterialElement->SetLabel( *Material->GetName() );
+
 					ImportContext.ImportedMaterials.Add( MaterialElement, Material );
+					ImportContext.Scene->AddMaterial( MaterialElement );
 				}
 			}
 		}
diff --git a/Engine/Plugins/Enterprise/DatasmithImporter/Source/DatasmithImporter/Private/Utility/DatasmithImporterUtils.cpp b/Engine/Plugins/Enterprise/DatasmithImporter/Source/DatasmithImporter/Private/Utility/DatasmithImporterUtils.cpp
--- a/Engine/Plugins/Enterprise/DatasmithImporter/Source/DatasmithImporter/Private/Utility/DatasmithImporterUtils.cpp
+++ b/Engine/Plugins/Enterprise/DatasmithImporter/Source/DatasmithImporter/Private/Utility/DatasmithImporterUtils.cpp
@@ -467,13 +467,6 @@ FString FDatasmithImporterUtils::GetDatasmithElementIdString(UObject* Object)
 
 namespace FDatasmithImporterUtilsHelper
 {
-	// #ueent_todo: Replace this with using FDatasmithImporterUtils::UniqueIdMetaDataKey
-	FString GetObjectTag(UObject* Object)
-	{
-		FString ObjectPath = FPaths::Combine( Object->GetOutermost()->GetName(), Object->GetName() );
-		return FMD5::HashAnsiString( *ObjectPath );
-	}
-
 	void SetupPointLightElement( const UPointLightComponent* PointLightComponent, IDatasmithPointLightElement* PointLightElement )
 	{
 		switch ( PointLightComponent->IntensityUnits )
@@ -570,7 +563,7 @@ namespace FDatasmithImporterUtilsHelper
 
 		if( LightComponent->LightFunctionMaterial != nullptr )
 		{
-			FString MaterialTag = GetObjectTag( LightComponent->LightFunctionMaterial );
+			FString MaterialTag = FDatasmithImporterUtils::GetDatasmithElementIdString( LightComponent->LightFunctionMaterial );
 			TSharedPtr< IDatasmithMaterialIDElement > MaterialIDElement = FDatasmithSceneFactory::CreateMaterialId( *MaterialTag );
 			LightActorElement->SetLightFunctionMaterial( MaterialIDElement );
 		}
@@ -670,6 +663,11 @@ namespace FDatasmithImporterUtilsHelper
 			return TSharedPtr< IDatasmithActorElement >();
 		}
 
+		FSoftObjectPath LightShapeBlueprintRef = FSoftObjectPath( TEXT("/DatasmithContent/Datasmith/DatasmithArealight.DatasmithArealight") );
+		UBlueprint* LightShapeBlueprint = Cast< UBlueprint >( LightShapeBlueprintRef.TryLoad() );
+
+		bool bNeedsTemplates = FDatasmithObjectTemplateUtils::GetObjectTemplate<UDatasmithActorTemplate>( Actor ) == nullptr;
+
 		auto AddTemplate = [](UClass* TemplateClass, UObject* Source, UObject* Outer)
 		{
 			UDatasmithObjectTemplate* DatasmithTemplate = NewObject< UDatasmithObjectTemplate >( Outer, TemplateClass );
@@ -677,28 +675,10 @@ namespace FDatasmithImporterUtilsHelper
 			FDatasmithObjectTemplateUtils::SetObjectTemplate( Outer, DatasmithTemplate );
 		};
 
-		FSoftObjectPath LightShapeBlueprintRef = FSoftObjectPath( TEXT("/DatasmithContent/Datasmith/DatasmithArealight.DatasmithArealight") );
-		UBlueprint* LightShapeBlueprint = Cast< UBlueprint >( LightShapeBlueprintRef.TryLoad() );
-		bool bNeedsTemplates = FDatasmithObjectTemplateUtils::GetObjectTemplate<UDatasmithActorTemplate>( Actor ) == nullptr;
-
-		FString ActorName = Actor->GetName();
-
-		if(bNeedsTemplates)
+		auto CreateMeshActorElement = [&](FString& ElementName, UStaticMeshComponent* StaticMeshComponent)
 		{
-			AddTemplate( UDatasmithActorTemplate::StaticClass(), Actor, Actor );
-
-			USceneComponent* RootComponent = Actor->GetRootComponent();
-			AddTemplate( UDatasmithSceneComponentTemplate::StaticClass(), RootComponent, RootComponent->GetOwner() );
-			UDatasmithAssetUserData::SetDatasmithUserDataValueForKey(RootComponent, UDatasmithAssetUserData::UniqueIdMetaDataKey, Actor->GetName() );
-		}
+			TSharedPtr< IDatasmithMeshActorElement > StaticMeshActorElement = FDatasmithSceneFactory::CreateMeshActor( *ElementName );
 
-		TSharedPtr< IDatasmithActorElement > ActorElement;
-		// #ueent_todo: Add proper support for all type of actors
-		if( AStaticMeshActor* StaticMeshActor = Cast<AStaticMeshActor>(Actor) )
-		{
-			TSharedPtr< IDatasmithMeshActorElement > StaticMeshActorElement = FDatasmithSceneFactory::CreateMeshActor( *ActorName );
-
-			UStaticMeshComponent* StaticMeshComponent = StaticMeshActor->GetStaticMeshComponent();
 			for( UMaterialInterface* MaterialInterface : StaticMeshComponent->OverrideMaterials )
 			{
 				if ( MaterialInterface )
@@ -710,7 +690,7 @@ namespace FDatasmithImporterUtilsHelper
 
 			if ( UStaticMesh* StaticMesh = StaticMeshComponent->GetStaticMesh() )
 			{
-				FString StaticMeshTag = GetObjectTag(StaticMeshComponent->GetStaticMesh());
+				FString StaticMeshTag = FDatasmithImporterUtils::GetDatasmithElementIdString( StaticMesh );
 				StaticMeshActorElement->SetStaticMeshPathName(*StaticMeshTag);
 			}
 
@@ -719,7 +699,25 @@ namespace FDatasmithImporterUtilsHelper
 				AddTemplate( UDatasmithStaticMeshComponentTemplate::StaticClass(), StaticMeshComponent, StaticMeshComponent );
 			}
 
-			ActorElement = StaticMeshActorElement;
+			return StaticMeshActorElement;
+		};
+
+		FString ActorName = Actor->GetName();
+
+		if(bNeedsTemplates)
+		{
+			AddTemplate( UDatasmithActorTemplate::StaticClass(), Actor, Actor );
+
+			USceneComponent* RootComponent = Actor->GetRootComponent();
+			AddTemplate( UDatasmithSceneComponentTemplate::StaticClass(), RootComponent, RootComponent->GetOwner() );
+			UDatasmithAssetUserData::SetDatasmithUserDataValueForKey(RootComponent, UDatasmithAssetUserData::UniqueIdMetaDataKey, Actor->GetName() );
+		}
+
+		TSharedPtr< IDatasmithActorElement > ActorElement;
+		// #ueent_todo: Add proper support for all type of actors
+		if( AStaticMeshActor* StaticMeshActor = Cast<AStaticMeshActor>(Actor) )
+		{
+			ActorElement = CreateMeshActorElement( ActorName, StaticMeshActor->GetStaticMeshComponent() );
 		}
 		else if( ADatasmithAreaLightActor* AreaLightActor = Cast<ADatasmithAreaLightActor>(Actor) )
 		{
@@ -824,7 +822,7 @@ namespace FDatasmithImporterUtilsHelper
 						HISMActorElement->AddMaterialOverride( *MaterialInterface->GetName(), 0 );
 					}
 
-					FString StaticMeshTag = GetObjectTag( HISMComponent->GetStaticMesh() );
+					FString StaticMeshTag = FDatasmithImporterUtils::GetDatasmithElementIdString( HISMComponent->GetStaticMesh() );
 					HISMActorElement->SetStaticMeshPathName( *StaticMeshTag );
 
 					ActorElement = HISMActorElement;
@@ -832,7 +830,17 @@ namespace FDatasmithImporterUtilsHelper
 				break;
 			case EDatasmithElementType::None:
 			default:
-				ActorElement = FDatasmithSceneFactory::CreateActor( *ActorName );
+				{
+					ActorElement = FDatasmithSceneFactory::CreateActor( *ActorName );
+					TArray<UStaticMeshComponent*> MeshComponents;
+					Actor->GetComponents<UStaticMeshComponent>( MeshComponents );
+					for(UStaticMeshComponent* MeshComponent : MeshComponents)
+					{
+						TSharedPtr< IDatasmithMeshActorElement > MeshActorElement = CreateMeshActorElement( ActorName, MeshComponent );
+						MeshActorElement->SetIsAComponent(true);
+						ActorElement->AddChild( MeshActorElement );
+					}
+				}
 				break;
 			}
 		}
diff --git a/Engine/Plugins/Experimental/DataPrepEditor/Source/DataPrepCore/Private/Parameterization/DataprepParameterization.cpp b/Engine/Plugins/Experimental/DataPrepEditor/Source/DataPrepCore/Private/Parameterization/DataprepParameterization.cpp
--- a/Engine/Plugins/Experimental/DataPrepEditor/Source/DataPrepCore/Private/Parameterization/DataprepParameterization.cpp
+++ b/Engine/Plugins/Experimental/DataPrepEditor/Source/DataPrepCore/Private/Parameterization/DataprepParameterization.cpp
@@ -10,10 +10,12 @@
 #include "Engine/Engine.h"
 #include "Math/UnrealMathUtility.h"
 #include "Parameterization/DataprepParameterizationUtils.h"
-#include "Serialization/ObjectReader.h"
+#include "Serialization/Archive.h"
 #include "Templates/UnrealTemplate.h"
 #include "UObject/Object.h"
 #include "UObject/PropertyPortFlags.h"
+#include "UObject/StructOnScope.h"
+#include "UObject/UObjectGlobals.h"
 #include "UObject/UObjectHash.h"
 #include "UObject/UnrealType.h"
 
@@ -322,6 +324,7 @@ bool FDataprepParameterizationBinding::operator==(const FDataprepParameterizatio
 	return ObjectBinded == Other.ObjectBinded && GetTypeHash(PropertyChain) == GetTypeHash(Other.PropertyChain) && ValueType == Other.ValueType;
 }
 
+
 uint32 GetTypeHash(const FDataprepParameterizationBinding& Binding)
 {
 	uint32 Hash = GetTypeHash( Binding.ObjectBinded );
@@ -340,6 +343,239 @@ uint32 GetTypeHash(const TArray<FDataprepPropertyLink>& PropertyLinks)
 	return Hash;
 }
 
+
+bool UDataprepParameterizationBindings::ContainsBinding(const TSharedRef<FDataprepParameterizationBinding>& Binding) const
+{
+	return BindingToParameterName.Contains( Binding );
+}
+
+bool UDataprepParameterizationBindings::HasBindingsForParameter(FName ParameterName) const
+{
+	return NameToBindings.Contains( ParameterName );
+}
+
+void UDataprepParameterizationBindings::Add(const TSharedRef<FDataprepParameterizationBinding>& Binding, FName ParamerterName)
+{
+	Modify();
+
+	uint32 BindingHash = GetTypeHash( Binding.Get() );
+
+	if ( FName* ExistingParameterName = BindingToParameterName.FindByHash( BindingHash, Binding ) )
+	{
+		if ( *ExistingParameterName != ParamerterName )
+		{
+			// Remove the trace of the old mapping
+			NameToBindings.FindRef( *ExistingParameterName ).RemoveByHash( BindingHash, Binding );
+		}
+	}
+
+	BindingToParameterName.AddByHash( BindingHash, Binding, ParamerterName );
+	NameToBindings.FindOrAdd( ParamerterName ).AddByHash( BindingHash, Binding );
+	ObjectToBindings.FindOrAdd( Binding->ObjectBinded ).AddByHash( BindingHash, Binding );
+}
+
+FName UDataprepParameterizationBindings::RemoveBinding(const TSharedRef<FDataprepParameterizationBinding>& Binding)
+{
+	uint32 BindingHash = GetTypeHash( Binding.Get() );
+	if ( FName* ParameterName = BindingToParameterName.FindByHash( BindingHash, Binding ) )
+	{
+		Modify();
+
+		// 1) Remove from the map binding to parameter
+		BindingToParameterName.RemoveByHash( BindingHash, Binding );
+		
+		// 2) Remove from the map parameter to bindings
+		{
+			if ( FSetOfBinding* BindingsMappedToParameter = NameToBindings.Find( *ParameterName ) )
+			{
+				if ( BindingsMappedToParameter->Num() <= 1 )
+				{
+					NameToBindings.Remove( *ParameterName );
+				}
+				else
+				{
+					BindingsMappedToParameter->RemoveByHash( BindingHash, Binding );
+				}
+			}
+		}
+
+		// 3) Remove from the map object to bindings
+		{
+			UObject* Object = Binding->ObjectBinded;
+			uint32 ObjectHash = GetTypeHash( Object );
+			if ( FSetOfBinding* BindingsMappedToObject = ObjectToBindings.FindByHash( ObjectHash, Object ) )
+			{
+				if ( BindingsMappedToObject->Num() <= 1 )
+				{
+					ObjectToBindings.RemoveByHash( ObjectHash, Object );
+				}
+				else
+				{
+					BindingsMappedToObject->RemoveByHash( BindingHash, Binding );
+				}
+			}
+		}
+
+		return *ParameterName;
+	}
+
+	return NAME_None;
+}
+
+const UDataprepParameterizationBindings::FBindingToParameterNameMap& UDataprepParameterizationBindings::GetBindingToParameterName() const
+{
+	return BindingToParameterName;
+}
+
+void UDataprepParameterizationBindings::Serialize(FArchive& Ar)
+{
+	Super::Serialize( Ar );
+
+	if ( Ar.IsSaving() )
+	{
+		Save( Ar );
+	}
+	else if ( Ar.IsLoading() )
+	{
+		Load( Ar );
+	}
+}
+
+void UDataprepParameterizationBindings::AddReferencedObjects(UObject* InThis, FReferenceCollector& Collector)
+{
+	UDataprepParameterizationBindings* Bindings = CastChecked<UDataprepParameterizationBindings>( InThis );
+
+	for ( TPair<UObject*, FSetOfBinding>& Pair : Bindings->ObjectToBindings )
+	{
+		Collector.AddReferencedObject( Pair.Key );
+
+		FSetOfBinding& BindingSet = Pair.Value;
+		for ( TSharedRef<FDataprepParameterizationBinding>& Binding : BindingSet )
+		{
+			Collector.AddReferencedObject( Binding->ObjectBinded );
+			Collector.AddReferencedObject( Binding->ValueType );
+		}
+	}
+
+	Super::AddReferencedObjects( InThis, Collector );
+}
+
+void UDataprepParameterizationBindings::Save(FArchive& Ar)
+{
+	// 1) Save the data for the name to bindings
+	{
+		int32 NumberOfParameter = NameToBindings.Num();
+		Ar << NumberOfParameter;
+		for ( const TPair<FName, FSetOfBinding>& PairNameAndBindings : NameToBindings )
+		{
+			FName ParameterName = PairNameAndBindings.Key;
+			Ar << ParameterName;
+			int32 NumberOfBinding = PairNameAndBindings.Value.Num();
+			Ar << NumberOfBinding;
+		}
+	}
+
+	// 2) Save the data for the object to bindings
+	{
+		int32 NumberOfObjectBinded = ObjectToBindings.Num();
+		Ar << NumberOfObjectBinded;
+		for ( const TPair<UObject*, FSetOfBinding>& PairObjectAndBindings : ObjectToBindings )
+		{
+			UObject* Object = PairObjectAndBindings.Key;
+			Ar << Object;
+			int32 NumberOfBinding = PairObjectAndBindings.Value.Num();
+			Ar << NumberOfBinding;
+		}
+	}
+
+	// 3) Do the actual save of the bindings
+	{
+		int32 NumberOfBinding = BindingToParameterName.Num();
+		Ar << NumberOfBinding;
+		UScriptStruct* StaticStruct = FDataprepParameterizationBinding::StaticStruct();
+
+		for (const TPair<TSharedRef<FDataprepParameterizationBinding>, FName>& PairBindingToParameterName : BindingToParameterName)
+		{
+			FDataprepParameterizationBinding& Binding = PairBindingToParameterName.Key.Get();
+			StaticStruct->SerializeItem( Ar, &Binding, nullptr );
+
+			FName ParameterName =  PairBindingToParameterName.Value;
+			Ar << ParameterName;
+		}
+	}
+}
+
+void UDataprepParameterizationBindings::Load(FArchive& Ar)
+{
+	// 1) Load the data for the name to bindings (prepare the structure for the insertion of data later)
+	{
+		int32 NumOfUniqueParameterName;
+		Ar << NumOfUniqueParameterName;
+		NameToBindings.Empty( NumOfUniqueParameterName );
+
+		for ( int32 Index = 0; Index < NumOfUniqueParameterName; Index++ )
+		{
+			FName ParameterName;
+			Ar << ParameterName;
+			FSetOfBinding& Bindings = NameToBindings.FindOrAdd( ParameterName );
+
+			int32 NumOfBindingForParameter;
+			Ar << NumOfBindingForParameter;
+			Bindings.Reserve( NumOfBindingForParameter );
+		}
+	}
+
+	// 2) Load the data for the object to bindings (prepare the structure for the insertion of data later)
+	{
+		int32 NumOfObjectBinded;
+		Ar << NumOfObjectBinded;
+		ObjectToBindings.Empty( NumOfObjectBinded );
+
+		for ( int32 Index = 0; Index < NumOfObjectBinded; Index++ )
+		{
+			UObject* Object = nullptr;
+			Ar << Object;
+
+			FSetOfBinding& Bindings = ObjectToBindings.FindOrAdd( Object );
+
+			int32 NumOfBindingForObject;
+			Ar << NumOfBindingForObject;
+			Bindings.Reserve( NumOfBindingForObject );
+		}
+	}
+
+	// 3) Do the actual loading of the binding
+	{
+		int32 NumOfBindings;
+		Ar << NumOfBindings;
+		BindingToParameterName.Empty( NumOfBindings );
+
+		UScriptStruct* StaticStruct = FDataprepParameterizationBinding::StaticStruct();
+
+		for (int32 Index = 0; Index < NumOfBindings; Index++)
+		{
+			TSharedRef<FDataprepParameterizationBinding> BindingPtr = MakeShared<FDataprepParameterizationBinding>();
+			FDataprepParameterizationBinding& Binding = BindingPtr.Get();
+			StaticStruct->SerializeItem( Ar, &Binding, nullptr );
+
+			FName ParameterizationName;
+			Ar << ParameterizationName;
+
+			if ( Binding.ObjectBinded )
+			{
+				uint32 BindingHash = GetTypeHash( Binding );
+
+				BindingToParameterName.AddByHash( BindingHash, BindingPtr, ParameterizationName );
+
+				NameToBindings.FindOrAdd( ParameterizationName ).AddByHash( BindingHash, BindingPtr );
+
+				ObjectToBindings.FindOrAdd( BindingPtr->ObjectBinded ).AddByHash( BindingHash, MoveTemp( BindingPtr ) );
+			}
+		}
+	}
+}
+
+
 UDataprepParameterization::UDataprepParameterization()
 {
 	OnObjectModifiedHandle = FCoreUObjectDelegates::OnObjectModified.AddUObject( this, &UDataprepParameterization::OnObjectModified );
@@ -356,6 +592,13 @@ void UDataprepParameterization::PostInitProperties()
 
 	if ( !HasAnyFlags( RF_ClassDefaultObject | RF_NeedLoad ) )
 	{
+		SetFlags( RF_Public );
+
+		if ( !BindingsContainer )
+		{
+			BindingsContainer = NewObject<UDataprepParameterizationBindings>( this, NAME_None, RF_Transactional );
+		}
+
 		GenerateClass();
 	}
 }
@@ -365,6 +608,12 @@ void UDataprepParameterization::PostLoad()
 	if ( !HasAnyFlags( RF_ClassDefaultObject | RF_NeedLoad ) )
 	{
 		SetFlags( RF_Public );
+
+		if ( !BindingsContainer )
+		{
+			BindingsContainer = NewObject<UDataprepParameterizationBindings>(this, NAME_None, RF_Transactional);
+		}
+
 		LoadParameterization();
 	}
 	Super::PostLoad();
@@ -418,25 +667,21 @@ bool UDataprepParameterization::BindObjectProperty(UObject* InObject, const TArr
 {
 	if ( InObject && InPropertyChain.Num() > 0 && !ParameterName.IsNone() )
 	{
-		FDataprepParameterizationBinding Binding( InObject, InPropertyChain );
+		TSharedRef<FDataprepParameterizationBinding> Binding = MakeShared<FDataprepParameterizationBinding>( InObject, InPropertyChain );
 		void* AddressOfTheValueFromBinding;
 
 		// We expect the chain to have a valid chain of cached property before inserting the binding
-		if ( DataprepParameterization::GetDeepestLevelOfValidCache( Binding, AddressOfTheValueFromBinding) == InPropertyChain.Num() - 1 )
+		if ( DataprepParameterization::GetDeepestLevelOfValidCache( Binding.Get(), AddressOfTheValueFromBinding) == InPropertyChain.Num() - 1 )
 		{
 			if ( UProperty** PropertyPtr = NameToParameterizationProperty.Find( ParameterName ) )
 			{
 				UProperty* Property = *PropertyPtr;
 
-
 				// Ensure that the properties type match
-				if ( Property->GetClass() == Binding.ValueType && !BindingsFromPipeline.Contains( Binding ) )
+				if ( Property->GetClass() == Binding->ValueType )
 				{
 					Modify();
-
-					// Valid bind the add binding to the name
-					BindingsFromPipeline.Add( Binding, ParameterName );
-					NameUsage.FindOrAdd( ParameterName ).Add( Binding );
+					BindingsContainer->Add( Binding, ParameterName );
 				}
 			}
 			else
@@ -444,8 +689,8 @@ bool UDataprepParameterization::BindObjectProperty(UObject* InObject, const TArr
 				Modify();
 				OnCustomClassAboutToBeUpdated.Broadcast();
 
-				BindingsFromPipeline.Add( MoveTemp(Binding) , ParameterName );
-				NameUsage.FindOrAdd( ParameterName ).Add( FDataprepParameterizationBinding(InObject, InPropertyChain) );
+				BindingsContainer->Add( Binding, ParameterName );
+			
 				UProperty* SourceProperty = InPropertyChain.Last().CachedProperty.Get();
 
 				// The validation we did with GetDeepestLevelOfValidCache ensure us that the property ptr is valid
@@ -467,30 +712,22 @@ bool UDataprepParameterization::BindObjectProperty(UObject* InObject, const TArr
 
 bool UDataprepParameterization::IsObjectPropertyBinded(UObject* Object, const TArray<FDataprepPropertyLink>& PropertyChain) const
 {
-	FDataprepParameterizationBinding Binding( Object, PropertyChain );
-	return BindingsFromPipeline.Contains( Binding );
+	TSharedRef<FDataprepParameterizationBinding> Binding= MakeShared<FDataprepParameterizationBinding>( Object, PropertyChain );
+	return BindingsContainer->ContainsBinding( Binding );
 }
 
 void UDataprepParameterization::RemoveBindedObjectProperty(UObject* Object, const TArray<FDataprepPropertyLink>& PropertyChain)
 {
-	FDataprepParameterizationBinding Binding( Object, PropertyChain );
-	uint32 Hash = GetTypeHash( Binding );
-	FName* ParameterNamePtr = BindingsFromPipeline.FindByHash( Hash, Binding );
-	if ( ParameterNamePtr )
+	TSharedRef<FDataprepParameterizationBinding> Binding = MakeShared<FDataprepParameterizationBinding>( Object, PropertyChain );
+
+	Modify();
+
+	FName ParameterOfRemovedBinding = BindingsContainer->RemoveBinding( Binding );
+	if ( !BindingsContainer->HasBindingsForParameter( ParameterOfRemovedBinding ) )
 	{
-		Modify();
 		OnCustomClassAboutToBeUpdated.Broadcast();
-		BindingsFromPipeline.RemoveByHash( Hash, Binding );
-		FName ParameterName = *ParameterNamePtr; 
-		if ( TSet<FDataprepParameterizationBinding>* BindingAssociatedToName = NameUsage.Find( ParameterName ) )
-		{
-			BindingAssociatedToName->RemoveByHash( Hash, Binding );
-			if ( BindingAssociatedToName->Num() == 0 )
-			{
-				NameToParameterizationProperty.Remove( ParameterName );
-				UpdateClass();
-			}
-		}
+		NameToParameterizationProperty.Remove( ParameterOfRemovedBinding );
+		UpdateClass();
 	}
 }
 
@@ -498,8 +735,8 @@ void UDataprepParameterization::GenerateClass()
 {
 	if ( !CustomContainerClass )
 	{
-		CustomContainerClass = NewObject<UClass>(GetOutermost(), FName(TEXT("Parameterization")), RF_Transient );
-		CustomContainerClass->SetSuperStruct(UObject::StaticClass());
+		CustomContainerClass = NewObject<UClass>( GetOutermost(), FName( TEXT("Parameterization") ), RF_Transient );
+		CustomContainerClass->SetSuperStruct( UObject::StaticClass() );
 
 		// Make the properties appear in a alphabetically order (for that we must add the properties to the class in the reverse order)
 		NameToParameterizationProperty.KeySort([](const FName& First, const FName& Second)
@@ -509,8 +746,8 @@ void UDataprepParameterization::GenerateClass()
 
 		for ( TPair<FName, UProperty*>& Pair : NameToParameterizationProperty )
 		{
-			UProperty* NewProperty = DuplicateObject<UProperty>(Pair.Value, CustomContainerClass, Pair.Key);
-			NewProperty->SetFlags(RF_Transient);
+			UProperty* NewProperty = DuplicateObject<UProperty>( Pair.Value, CustomContainerClass, Pair.Key );
+			NewProperty->SetFlags( RF_Transient );
 			NewProperty->PropertyFlags = CPF_Edit;
 
 			// Need to manually call Link to fix-up some data (such as the C++ property flags) that are only set during Link
@@ -519,16 +756,16 @@ void UDataprepParameterization::GenerateClass()
 				NewProperty->LinkWithoutChangingOffset(Ar);
 			}
 
-			CustomContainerClass->AddCppProperty(NewProperty);
+			CustomContainerClass->AddCppProperty( NewProperty );
 
 			Pair.Value = NewProperty;
 		}
 
 		CustomContainerClass->Bind();
-		CustomContainerClass->StaticLink(true);
-		CustomContainerClass->AssembleReferenceTokenStream(true);
+		CustomContainerClass->StaticLink( true );
+		CustomContainerClass->AssembleReferenceTokenStream( true );
 
-		DefaultParameterisation = CustomContainerClass->GetDefaultObject(true);
+		DefaultParameterisation = CustomContainerClass->GetDefaultObject( true );
 	}
 }
 
@@ -561,16 +798,14 @@ void UDataprepParameterization::LoadParameterization()
 		CustomContainerClass = NewObject<UClass>( GetOutermost(), FName(TEXT("Parameterization")), RF_Transient );
 		CustomContainerClass->SetSuperStruct( UObject::StaticClass() );
 
-		TSet<FDataprepParameterizationBinding> BindingToRemove;
-		NameUsage.Empty( NameUsage.Num() );
+		TSet<TSharedRef<FDataprepParameterizationBinding>> BindingToRemove;
 		NameToParameterizationProperty.Empty( NameToParameterizationProperty.Num() );
 
-		for ( TPair<FDataprepParameterizationBinding, FName>& Binding : BindingsFromPipeline )
+		for ( const TPair<TSharedRef<FDataprepParameterizationBinding>, FName>& Binding : BindingsContainer->GetBindingToParameterName() )
 		{
 			const FName BindingName = Binding.Value;
-			NameUsage.FindOrAdd( BindingName ).Add( Binding.Key );
 
-			UProperty* PropertyFromChain = DataprepParameterization::GetPropertyFromBinding( Binding.Key );
+			UProperty* PropertyFromChain = DataprepParameterization::GetPropertyFromBinding( Binding.Key.Get() );
 			UProperty** PropertyFromParameterizationClass = NameToParameterizationProperty.Find( BindingName );
 
 			if ( PropertyFromChain && !PropertyFromParameterizationClass )
@@ -579,14 +814,14 @@ void UDataprepParameterization::LoadParameterization()
 			}
 			else if ( !PropertyFromChain || PropertyFromChain->GetClass() != (*PropertyFromParameterizationClass)->GetClass() )
 			{
-				BindingToRemove.Add(Binding.Key);
+				BindingToRemove.Add( Binding.Key ) ;
 			}
 		}
 
 		// Remove the invalid bindings
-		for ( const FDataprepParameterizationBinding& InvalidBinding : BindingToRemove )
+		for ( const TSharedRef<FDataprepParameterizationBinding>& InvalidBinding : BindingToRemove )
 		{
-			BindingsFromPipeline.Remove( InvalidBinding );
+			BindingsContainer->RemoveBinding( InvalidBinding );
 		}
 
 		// Make the properties appear in a alphabetically order (for that we must add the properties to the class in the reverse order)
@@ -750,20 +985,19 @@ void UDataprepParameterizationInstance::OnObjectModified(UObject* Object)
 void UDataprepParameterizationInstance::ApplyParameterization(const TMap<UObject*, UObject*>& SourceToCopy)
 {
 	check( SourceParameterization );
-	for ( TPair<FDataprepParameterizationBinding, FName>& BindingPair : SourceParameterization->BindingsFromPipeline )
+	for ( const TPair<TSharedRef<FDataprepParameterizationBinding>, FName>& BindingPair : SourceParameterization->BindingsContainer->GetBindingToParameterName() )
 	{
-		FDataprepParameterizationBinding& Binding = BindingPair.Key;
-		if ( UObject* Object = SourceToCopy.FindRef( Binding.ObjectBinded ) )
+		const TSharedRef<FDataprepParameterizationBinding>& Binding = BindingPair.Key;
+		if ( UObject* Object = SourceToCopy.FindRef( Binding->ObjectBinded ) )
 		{
-			TGuardValue<UObject*> GuardObjectBinded( Binding.ObjectBinded, Object );
+			TGuardValue<UObject*> GuardObjectBinded( Binding->ObjectBinded, Object );
 
 			void* DestinationAddress = nullptr;
-			if ( UProperty* DestinationProperty = DataprepParameterization::GetPropertyFromBinding( Binding, DestinationAddress ) )
+			if ( UProperty* DestinationProperty = DataprepParameterization::GetPropertyFromBinding( Binding.Get(), DestinationAddress ) )
 			{
 				UProperty* ParameterizationProperty = FindObjectFast<UProperty>( SourceParameterization->CustomContainerClass, BindingPair.Value );
 				void* ParameterizationAddress =  DataprepParameterization::GetAddressOf( *ParameterizationProperty, ParameterizationInstance, INDEX_NONE );
 				DataprepParameterization::CopyCompleteValue( *DestinationProperty, DestinationAddress, *ParameterizationProperty, ParameterizationAddress );
-				check( Object );
 			}
 		}
 	}
diff --git a/Engine/Plugins/Experimental/DataPrepEditor/Source/DataPrepCore/Private/Parameterization/DataprepParameterization.h b/Engine/Plugins/Experimental/DataPrepEditor/Source/DataPrepCore/Private/Parameterization/DataprepParameterization.h
--- a/Engine/Plugins/Experimental/DataPrepEditor/Source/DataPrepCore/Private/Parameterization/DataprepParameterization.h
+++ b/Engine/Plugins/Experimental/DataPrepEditor/Source/DataPrepCore/Private/Parameterization/DataprepParameterization.h
@@ -4,6 +4,7 @@
 
 #include "Parameterization/DataprepParameterizationUtils.h"
 
+#include "Containers/ContainerAllocationPolicies.h"
 #include "CoreMinimal.h"
 #include "Delegates/DelegateCombinations.h"
 #include "Delegates/IDelegateInstance.h"
@@ -55,22 +56,112 @@ uint32 GetTypeHash(const FDataprepParameterizationBinding& Binding);
 
 uint32 GetTypeHash(const TArray<FDataprepPropertyLink>& PropertyLinks);
 
+/** 
+ * A override to defines how the map's pairs are hashed and compared.
+ * This allow us to have some map that work with TSharedRef but compare the object value instead of the pointer
+ */
+struct FDataprepParametrizationBindingMapKeyFuncs : TDefaultMapHashableKeyFuncs<TSharedRef<FDataprepParameterizationBinding>, FName, false>
+{
+	static FORCEINLINE bool Matches(KeyInitType A, KeyInitType B)
+	{
+		return A.Get() == B.Get();
+	}
+
+	static FORCEINLINE uint32 GetKeyHash(KeyInitType Key)
+	{
+		return GetTypeHash( Key.Get() );
+	}
+};
+
+/** 
+ * A override to defines how the set values are hashed and compared.
+ * This allow us to have some set that work with TSharedRef but compare the object value instead of the pointer
+ */
+struct FDataprepParametrizationBindingSetKeyFuncs : DefaultKeyFuncs<TSharedRef<FDataprepParameterizationBinding>>
+{
+	template<typename ComparableKey>
+	static FORCEINLINE bool Matches(KeyInitType A, ComparableKey B)
+	{
+		return A.Get() == B.Get();
+	}
+
+	static FORCEINLINE uint32 GetKeyHash(KeyInitType Key)
+	{
+		return GetTypeHash(Key.Get());
+	}
+};
+
+
+
 /**
- * Count the number of time each hash was encounter
+ * Encapsulate the unidirectionality necessary for a constant cost of access to the data related to the bindings
  */
-USTRUCT()
-struct FHashCount
+UCLASS(MinimalAPI, Experimental)
+class UDataprepParameterizationBindings : public UObject
 {
+public:
+
 	GENERATED_BODY()
 
-	UPROPERTY()
-	TMap<uint32, uint32> HashCount;
+	using FBindingToParameterNameMap = TMap<TSharedRef<FDataprepParameterizationBinding>, FName, FDefaultSetAllocator, FDataprepParametrizationBindingMapKeyFuncs>;
+
+	/**
+	 * Does the data structure contains this binding
+	 */
+	bool ContainsBinding(const TSharedRef<FDataprepParameterizationBinding>& Binding) const;
+
+	/**
+	 * Does the data structure has some bindings for the parameter name
+	 */
+	bool HasBindingsForParameter(FName ParameterName) const;
+
+	/**
+	 * Add a binding and map it to the parameter
+	 */
+	void Add(const TSharedRef<FDataprepParameterizationBinding>& Binding, FName ParameterName);
+
+	/**
+	 * Remove a binding.
+	 * @return The name of the parameter the binding was associated with
+	 */
+	FName RemoveBinding(const TSharedRef<FDataprepParameterizationBinding>& Binding);
+
+	const FBindingToParameterNameMap& GetBindingToParameterName() const;
+
+	// UObject interface
+	virtual void Serialize(FArchive& Ar) override;
+	static void AddReferencedObjects(UObject* InThis, FReferenceCollector& Collector);
+	// End of UObject interface
+
+private:
+	/**
+	 * Do the actual serialization when saving
+	 */
+	void Save(FArchive& Ar);
+
+	/**
+	 * Do the actual serialization when reloading
+	 */
+	void Load(FArchive& Ar);
+
+
+	/** Core storage also track a binding to it's parameter name */
+	FBindingToParameterNameMap BindingToParameterName;
+
+	using FSetOfBinding = TSet<TSharedRef<FDataprepParameterizationBinding>, FDataprepParametrizationBindingSetKeyFuncs>;
+
+	/** Track the name usage for parameters */
+	TMap<FName, FSetOfBinding> NameToBindings;
+
+	/** Track which binding a object has */
+	TMap<UObject*, FSetOfBinding> ObjectToBindings;
 };
 
+
 /** 
  * The DataprepParameterization contains the data for the parameterization of a pipeline
  */
-UCLASS(MinimalAPI)
+UCLASS(MinimalAPI, Experimental)
 class UDataprepParameterization : public UObject
 {
 public:
@@ -132,14 +223,12 @@ private:
 	 */
 	UProperty* AddPropertyToClass(FName ParameterisationPropertyName, UProperty& Property);
 
+	// The containers for the bindings
 	UPROPERTY()
-	TMap<FDataprepParameterizationBinding, FName> BindingsFromPipeline;
+	UDataprepParameterizationBindings* BindingsContainer;
 
 	UPROPERTY(Transient, NonTransactional)
-	TMap<FName, UProperty*> NameToParameterizationProperty; // Just a cache for the CustomFindProperty
-
-	/** Track the name usage for parameters */
-	TMap<FName, TSet<FDataprepParameterizationBinding>> NameUsage;
+	TMap<FName, UProperty*> NameToParameterizationProperty;
 
 	UPROPERTY(Transient, NonTransactional)
 	UClass* CustomContainerClass;
@@ -170,7 +259,7 @@ private:
 };
 
 
-UCLASS(MinimalAPI)
+UCLASS(MinimalAPI, Experimental)
 class UDataprepParameterizationInstance : public UObject
 {
 public:
