Merging CLs #9493596 and #9494437 from //UE4/Dev-Enterprise

9493596:
Fixed crash with DG files containing variants
Issue with merged meshes

9494437:
Removed dependency of DataprepEditor plugin on DatasmithContent

diff --git a/Engine/Plugins/Enterprise/DatasmithImporter/Source/DatasmithImporter/Private/Utility/DatasmithImporterUtils.cpp b/Engine/Plugins/Enterprise/DatasmithImporter/Source/DatasmithImporter/Private/Utility/DatasmithImporterUtils.cpp
--- a/Engine/Plugins/Enterprise/DatasmithImporter/Source/DatasmithImporter/Private/Utility/DatasmithImporterUtils.cpp
+++ b/Engine/Plugins/Enterprise/DatasmithImporter/Source/DatasmithImporter/Private/Utility/DatasmithImporterUtils.cpp
@@ -462,7 +462,15 @@ FName FDatasmithImporterUtils::GetDatasmithElementId( UObject* Object )
 
 FString FDatasmithImporterUtils::GetDatasmithElementIdString(UObject* Object)
 {
-	return UDatasmithAssetUserData::GetDatasmithUserDataValueForKey(Object, UDatasmithAssetUserData::UniqueIdMetaDataKey);
+	FString ElementId = UDatasmithAssetUserData::GetDatasmithUserDataValueForKey(Object, UDatasmithAssetUserData::UniqueIdMetaDataKey);
+
+	if( ElementId.IsEmpty() )
+	{
+		const FString ObjectPath = FPaths::Combine( Object->GetOutermost()->GetName(), Object->GetName() );
+		return FMD5::HashBytes( reinterpret_cast<const uint8*>(*ObjectPath), ObjectPath.Len() * sizeof(TCHAR) );
+	}
+
+	return ElementId;
 }
 
 namespace FDatasmithImporterUtilsHelper
diff --git a/Engine/Plugins/Experimental/DataPrepEditor/DataPrepEditor.uplugin b/Engine/Plugins/Experimental/DataPrepEditor/DataPrepEditor.uplugin
--- a/Engine/Plugins/Experimental/DataPrepEditor/DataPrepEditor.uplugin
+++ b/Engine/Plugins/Experimental/DataPrepEditor/DataPrepEditor.uplugin
@@ -32,10 +32,6 @@
 		}
 	],
 	"Plugins": [
-		{
-			"Name": "DatasmithContent",
-			"Enabled": true
-		},
 		{
 			"Name": "EditorScriptingUtilities",
 			"Enabled": true
diff --git a/Engine/Plugins/Experimental/DataPrepEditor/Source/DataPrepEditor/Private/DataprepSnapshot.cpp b/Engine/Plugins/Experimental/DataPrepEditor/Source/DataPrepEditor/Private/DataprepSnapshot.cpp
--- a/Engine/Plugins/Experimental/DataPrepEditor/Source/DataPrepEditor/Private/DataprepSnapshot.cpp
+++ b/Engine/Plugins/Experimental/DataPrepEditor/Source/DataPrepEditor/Private/DataprepSnapshot.cpp
@@ -49,6 +49,34 @@ namespace DataprepSnapshotUtil
 {
 	const TCHAR* SnapshotExtension = TEXT(".dpc");
 
+
+	/**
+	 * Extends FObjectAndNameAsStringProxyArchive to support FLazyObjectPtr.
+	 */
+	struct FSnapshotCustomArchive : public FObjectAndNameAsStringProxyArchive
+	{
+		FSnapshotCustomArchive(FArchive& InInnerArchive)
+			:	FObjectAndNameAsStringProxyArchive(InInnerArchive, false)
+		{
+			// Set archive as transacting to persist all data including data in memory
+			SetIsTransacting( true );
+		}
+
+		virtual FArchive& operator<<(FLazyObjectPtr& Obj) override
+		{
+			// Copied from FArchiveUObject::SerializeLazyObjectPtr
+			// Note that archive is transacting
+			if (IsLoading())
+			{
+				// Reset before serializing to clear the internal weak pointer. 
+				Obj.Reset();
+			}
+			InnerArchive << Obj.GetUniqueID();
+
+			return *this;
+		}
+	};
+
 	void RemoveSnapshotFiles(const FString& RootDir)
 	{
 		TRACE_CPUPROFILER_EVENT_SCOPE(DataprepSnapshotUtil::RemoveSnapshotFiles);
@@ -106,8 +134,7 @@ namespace DataprepSnapshotUtil
 		};
 
 		FMemoryWriter MemAr(OutSerializedData);
-		FObjectAndNameAsStringProxyArchive Ar(MemAr, false);
-		Ar.SetIsTransacting(true);
+		FSnapshotCustomArchive Ar(MemAr);
 
 		// Collect sub-objects depending on input object including nested objects
 		TArray< UObject* > SubObjectsArray;
@@ -265,8 +292,7 @@ namespace DataprepSnapshotUtil
 		RemoveDefaultDependencies( Object );
 
 		FMemoryReader MemAr(bUseCompression ? MemoryBuffer : InSerializedData);
-		FObjectAndNameAsStringProxyArchive Ar(MemAr, false);
-		Ar.SetIsTransacting(true);
+		FSnapshotCustomArchive Ar(MemAr);
 
 		// Deserialize count of sub-objects
 		int32 SubObjectsCount = 0;
diff --git a/Engine/Plugins/Experimental/DataPrepEditor/Source/DataPrepLibraries/DataPrepLibraries.Build.cs b/Engine/Plugins/Experimental/DataPrepEditor/Source/DataPrepLibraries/DataPrepLibraries.Build.cs
--- a/Engine/Plugins/Experimental/DataPrepEditor/Source/DataPrepLibraries/DataPrepLibraries.Build.cs
+++ b/Engine/Plugins/Experimental/DataPrepEditor/Source/DataPrepLibraries/DataPrepLibraries.Build.cs
@@ -13,7 +13,6 @@ namespace UnrealBuildTool.Rules
 					"Core",
 					"CoreUObject",
 					"DataprepCore",
-					"DatasmithContent",
 					"EditorScriptingUtilities",
 					"EditorStyle",
 					"Engine",
diff --git a/Engine/Plugins/Experimental/DataPrepEditor/Source/DataPrepLibraries/Private/DataPrepFilterLibrary.cpp b/Engine/Plugins/Experimental/DataPrepEditor/Source/DataPrepLibraries/Private/DataPrepFilterLibrary.cpp
--- a/Engine/Plugins/Experimental/DataPrepEditor/Source/DataPrepLibraries/Private/DataPrepFilterLibrary.cpp
+++ b/Engine/Plugins/Experimental/DataPrepEditor/Source/DataPrepLibraries/Private/DataPrepFilterLibrary.cpp
@@ -2,29 +2,10 @@
 
 #include "DataPrepFilterLibrary.h"
 
-#include "DatasmithContentBlueprintLibrary.h"
 #include "Engine/StaticMesh.h"
 
 namespace DataprepFilterLibraryImpl
 {
-	bool StringCompare( const FString& StringToCompare, const FString& SearchString, EEditorScriptingStringMatchType StringMatch )
-	{
-		switch ( StringMatch )
-		{
-			case EEditorScriptingStringMatchType::Contains:
-				return StringToCompare.Contains( SearchString );
-
-			case EEditorScriptingStringMatchType::ExactMatch:
-				return StringToCompare.Equals( SearchString );
-
-			case EEditorScriptingStringMatchType::MatchesWildcard:
-				return StringToCompare.MatchesWildcard( SearchString );
-
-			default:
-				return false;
-		}
-	}
-
 	template<typename T>
 	T* CastIfValid(UObject* Target)
 	{
@@ -43,23 +24,6 @@ TArray< UObject* > UDataprepFilterLibrary::FilterByName( const TArray< UObject*
 	return UEditorFilterLibrary::ByIDName( TargetArray, NameSubString, StringMatch, EEditorScriptingFilterType::Include );
 }
 
-TArray< UObject* > UDataprepFilterLibrary::FilterByMetadata( const TArray< UObject* >& TargetArray, FName Key, const FString& Value, EEditorScriptingStringMatchType ValueMatch )
-{
-	TArray< UObject* > Results;
-
-	for ( UObject* Object : TargetArray )
-	{
-		FString KeyValue = UDatasmithContentBlueprintLibrary::GetDatasmithUserDataValueForKey( Object, Key );
-
-		if ( !KeyValue.IsEmpty() && DataprepFilterLibraryImpl::StringCompare( KeyValue, Value, ValueMatch ) )
-		{
-			Results.Add( Object );
-		}
-	}
-
-	return Results;
-}
-
 TArray< UObject* > UDataprepFilterLibrary::FilterBySize(const TArray< UObject* >& TargetArray, EDataprepSizeSource SizeSource, EDataprepSizeFilterMode FilterMode, float Threshold)
 {
 	TArray< UObject* > Result;
diff --git a/Engine/Plugins/Experimental/DataPrepEditor/Source/DataPrepLibraries/Private/DataPrepFilterLibrary.h b/Engine/Plugins/Experimental/DataPrepEditor/Source/DataPrepLibraries/Private/DataPrepFilterLibrary.h
--- a/Engine/Plugins/Experimental/DataPrepEditor/Source/DataPrepLibraries/Private/DataPrepFilterLibrary.h
+++ b/Engine/Plugins/Experimental/DataPrepEditor/Source/DataPrepLibraries/Private/DataPrepFilterLibrary.h
@@ -45,17 +45,6 @@ public:
 	UFUNCTION(BlueprintPure, Category = "Dataprep | Filter", meta = (DeterminesOutputType = "TargetArray"))
 	static TArray< UObject* > FilterByName( const TArray< UObject* >& TargetArray, const FString& NameSubString, EEditorScriptingStringMatchType StringMatch = EEditorScriptingStringMatchType::Contains );
 
-	/**
-	 * Filter the array based on the metadata key and value.
-	 * @param	TargetArray		Array of Object to filter. The array will not change.
-	 * @param	Key				The key to filter with.
-	 * @param	Value			An optional value to filter with.
-	 * @param	ValueMatch		Contains the Value OR matches with the wildcard *? OR exactly the same value.
-	 * @return	The filtered list.
-	 */
-	UFUNCTION(BlueprintPure, Category = "Dataprep | Filter", meta = (DeterminesOutputType = "TargetArray"))
-	static TArray< UObject* > FilterByMetadata( const TArray< UObject* >& TargetArray, FName Key, const FString& Value, EEditorScriptingStringMatchType ValueMatch = EEditorScriptingStringMatchType::Contains );
-
 	/**
 	 * Filter the array based on the geometry size.
 	 * @param	TargetArray       Array of Actors or StaticMeshes to filter. The array will not change.
