Merging CLs #9542438, #9542473 and #9542633 from //UE4/Dev-Enterprise

9542438:
Implemented standard selection in the Dataprep viewport
Added ability to overwrite viewport's environment map
Fixed a crash in the DatasmithImporterUtils.cpp (check on a nullptr)
Modified Merge operation to set the name of the mesh and actor to the one provided by the user

9542473:
Dataprep Parametrization now remove the bindings associted to a object when it is removed from the dataprep pipeline.

diff --git a/Engine/Plugins/Enterprise/DatasmithImporter/Source/DatasmithImporter/Private/Utility/DatasmithImporterUtils.cpp b/Engine/Plugins/Enterprise/DatasmithImporter/Source/DatasmithImporter/Private/Utility/DatasmithImporterUtils.cpp
--- a/Engine/Plugins/Enterprise/DatasmithImporter/Source/DatasmithImporter/Private/Utility/DatasmithImporterUtils.cpp
+++ b/Engine/Plugins/Enterprise/DatasmithImporter/Source/DatasmithImporter/Private/Utility/DatasmithImporterUtils.cpp
@@ -464,7 +464,7 @@ FString FDatasmithImporterUtils::GetDatasmithElementIdString(UObject* Object)
 {
 	FString ElementId = UDatasmithAssetUserData::GetDatasmithUserDataValueForKey(Object, UDatasmithAssetUserData::UniqueIdMetaDataKey);
 
-	if( ElementId.IsEmpty() )
+	if( Object != nullptr && ElementId.IsEmpty() )
 	{
 		const FString ObjectPath = FPaths::Combine( Object->GetOutermost()->GetName(), Object->GetName() );
 		return FMD5::HashBytes( reinterpret_cast<const uint8*>(*ObjectPath), ObjectPath.Len() * sizeof(TCHAR) );
diff --git a/Engine/Plugins/Experimental/DataPrepEditor/Source/DataPrepCore/Private/DataPrepAsset.cpp b/Engine/Plugins/Experimental/DataPrepEditor/Source/DataPrepCore/Private/DataPrepAsset.cpp
--- a/Engine/Plugins/Experimental/DataPrepEditor/Source/DataPrepCore/Private/DataPrepAsset.cpp
+++ b/Engine/Plugins/Experimental/DataPrepEditor/Source/DataPrepCore/Private/DataPrepAsset.cpp
@@ -193,7 +193,7 @@ UObject* UDataprepAsset::GetParameterizationObject()
 	return Parameterization->GetDefaultObject();
 }
 
-void UDataprepAsset::BindObjectPropertyToParameterization(UObject* Object, const TArray<FDataprepPropertyLink>& InPropertyChain, FName Name)
+void UDataprepAsset::BindObjectPropertyToParameterization(UObject* Object, const TArray<FDataprepPropertyLink>& InPropertyChain, const FName& Name)
 {
 	bool bPassConditionCheck = false;
 
diff --git a/Engine/Plugins/Experimental/DataPrepEditor/Source/DataPrepCore/Private/DataprepActionAsset.cpp b/Engine/Plugins/Experimental/DataPrepEditor/Source/DataPrepCore/Private/DataprepActionAsset.cpp
--- a/Engine/Plugins/Experimental/DataPrepEditor/Source/DataPrepCore/Private/DataprepActionAsset.cpp
+++ b/Engine/Plugins/Experimental/DataPrepEditor/Source/DataPrepCore/Private/DataprepActionAsset.cpp
@@ -3,12 +3,14 @@
 #include "DataprepActionAsset.h"
 
 // Dataprep include
+#include "DataPrepAsset.h"
 #include "DataPrepOperation.h"
 #include "DataprepCoreLogCategory.h"
 #include "DataprepCorePrivateUtils.h"
 #include "DataprepCoreUtils.h"
 #include "IDataprepLogger.h"
 #include "IDataprepProgressReporter.h"
+#include "Parameterization/DataprepParameterization.h"
 #include "SelectionSystem/DataprepFilter.h"
 
 // Engine include
@@ -21,6 +23,7 @@
 #include "Materials/MaterialInterface.h"
 #include "Templates/SubclassOf.h"
 #include "UObject/ObjectMacros.h"
+#include "UObject/UObjectHash.h"
 
 #ifdef WITH_EDITOR
 #include "Editor.h"
@@ -41,10 +44,10 @@ UDataprepActionAsset::UDataprepActionAsset()
 
 		OperationContext->Context = MakeShareable( new FDataprepContext() );
 
-		OperationContext->AddAssetDelegate = FDataprepAddAsset::CreateUObject( this, &UDataprepActionAsset::OnAddAsset);
-		OperationContext->CreateActorDelegate = FDataprepCreateActor::CreateUObject( this, &UDataprepActionAsset::OnCreateActor);
-		OperationContext->RemoveObjectDelegate = FDataprepRemoveObject::CreateUObject( this, &UDataprepActionAsset::OnRemoveObject);
-		OperationContext->DeleteObjectsDelegate = FDataprepDeleteObjects::CreateUObject( this, &UDataprepActionAsset::OnDeleteObjects);
+		OperationContext->AddAssetDelegate = FDataprepAddAsset::CreateUObject( this, &UDataprepActionAsset::OnAddAsset );
+		OperationContext->CreateActorDelegate = FDataprepCreateActor::CreateUObject( this, &UDataprepActionAsset::OnCreateActor );
+		OperationContext->RemoveObjectDelegate = FDataprepRemoveObject::CreateUObject( this, &UDataprepActionAsset::OnRemoveObject );
+		OperationContext->DeleteObjectsDelegate = FDataprepDeleteObjects::CreateUObject( this, &UDataprepActionAsset::OnDeleteObjects );
 	}
 }
 
@@ -248,6 +251,17 @@ bool UDataprepActionAsset::RemoveStep(int32 Index)
 	if ( Steps.IsValidIndex( Index ) )
 	{
 		Modify();
+
+		if ( UDataprepAsset* DataprepAsset = FDataprepCoreUtils::GetDataprepAssetOfObject( this ) )
+		{
+			if ( UDataprepParameterization* Parameterization = DataprepAsset->GetDataprepParameterization() )
+			{
+				TArray< UObject* > Objects;
+				GetObjectsWithOuter( Steps[Index], Objects );
+				Parameterization->RemoveBindingFromObjects( Objects );
+			}
+		}
+
 		Steps.RemoveAt( Index );
 		OnStepsChanged.Broadcast();
 		return true;
@@ -263,6 +277,22 @@ FOnStepsOrderChanged& UDataprepActionAsset::GetOnStepsOrderChanged()
 	return OnStepsChanged;
 }
 
+void UDataprepActionAsset::NotifyDataprepSystemsOfRemoval()
+{
+	if ( UDataprepAsset* DataprepAsset = FDataprepCoreUtils::GetDataprepAssetOfObject( this ) )
+	{
+		if ( UDataprepParameterization * Parameterization = DataprepAsset->GetDataprepParameterization() )
+		{
+			TArray< UObject* > Objects;
+			for ( UDataprepActionStep* Step : Steps )
+			{
+				GetObjectsWithOuter( Step, Objects );
+			}
+			Parameterization->RemoveBindingFromObjects( Objects );
+		}
+	}
+}
+
 void UDataprepActionAsset::OnClassesRemoved(const TArray<UClass *>& DeletedClasses)
 {
 	for ( UClass* Class : DeletedClasses )
diff --git a/Engine/Plugins/Experimental/DataPrepEditor/Source/DataPrepCore/Private/DataprepCoreUtils.cpp b/Engine/Plugins/Experimental/DataPrepEditor/Source/DataPrepCore/Private/DataprepCoreUtils.cpp
--- a/Engine/Plugins/Experimental/DataPrepEditor/Source/DataPrepCore/Private/DataprepCoreUtils.cpp
+++ b/Engine/Plugins/Experimental/DataPrepEditor/Source/DataPrepCore/Private/DataprepCoreUtils.cpp
@@ -8,6 +8,7 @@
 #include "DataprepAssetUserData.h"
 #endif
 
+#include "DataPrepAsset.h"
 #include "DataprepCoreLogCategory.h"
 #include "IDataprepProgressReporter.h"
 
@@ -15,6 +16,7 @@
 #include "GameFramework/Actor.h"
 #include "LevelSequence.h"
 #include "Misc/ScopedSlowTask.h"
+#include "RenderingThread.h"
 #include "UObject/UObjectHash.h"
 
 #if WITH_EDITOR
@@ -23,10 +25,23 @@
 #include "ObjectTools.h"
 #include "Subsystems/AssetEditorSubsystem.h"
 #endif
-#include "RenderingThread.h"
 
 #define LOCTEXT_NAMESPACE "DataprepCoreUtils"
 
+UDataprepAsset* FDataprepCoreUtils::GetDataprepAssetOfObject(UObject* Object)
+{
+	while ( Object )
+	{
+		if ( UDataprepAsset::StaticClass() == Object->GetClass() )
+		{
+			return static_cast<UDataprepAsset*>( Object );
+		}
+		Object = Object->GetOuter();
+	}
+
+	return nullptr;
+}
+
 void FDataprepCoreUtils::PurgeObjects(TArray<UObject*> InObjects)
 {
 	TRACE_CPUPROFILER_EVENT_SCOPE(FDataprepCoreUtils::PurgeObjects);
diff --git a/Engine/Plugins/Experimental/DataPrepEditor/Source/DataPrepCore/Private/Parameterization/DataprepParameterization.cpp b/Engine/Plugins/Experimental/DataPrepEditor/Source/DataPrepCore/Private/Parameterization/DataprepParameterization.cpp
--- a/Engine/Plugins/Experimental/DataPrepEditor/Source/DataPrepCore/Private/Parameterization/DataprepParameterization.cpp
+++ b/Engine/Plugins/Experimental/DataPrepEditor/Source/DataPrepCore/Private/Parameterization/DataprepParameterization.cpp
@@ -14,7 +14,6 @@
 #include "Templates/UnrealTemplate.h"
 #include "UObject/Object.h"
 #include "UObject/PropertyPortFlags.h"
-#include "UObject/StructOnScope.h"
 #include "UObject/UObjectGlobals.h"
 #include "UObject/UObjectHash.h"
 #include "UObject/UnrealType.h"
@@ -349,12 +348,12 @@ bool UDataprepParameterizationBindings::ContainsBinding(const TSharedRef<FDatapr
 	return BindingToParameterName.Contains( Binding );
 }
 
-bool UDataprepParameterizationBindings::HasBindingsForParameter(FName ParameterName) const
+bool UDataprepParameterizationBindings::HasBindingsForParameter(const FName& ParameterName) const
 {
 	return NameToBindings.Contains( ParameterName );
 }
 
-void UDataprepParameterizationBindings::Add(const TSharedRef<FDataprepParameterizationBinding>& Binding, FName ParamerterName)
+void UDataprepParameterizationBindings::Add(const TSharedRef<FDataprepParameterizationBinding>& Binding, const FName& ParamerterName)
 {
 	Modify();
 
@@ -422,6 +421,26 @@ FName UDataprepParameterizationBindings::RemoveBinding(const TSharedRef<FDatapre
 	return NAME_None;
 }
 
+TSet<FName> UDataprepParameterizationBindings::RemoveAllBindingsFromObject(UObject* Object)
+{
+	if ( FSetOfBinding* Bindings = ObjectToBindings.Find( Object ) )
+	{
+		Modify();
+
+		// Remove the bindings
+		TSet<FName> ParameterName;
+		TArray<TSharedRef<FDataprepParameterizationBinding>> BindingsToRemove( Bindings->Array() );
+		for ( TSharedRef<FDataprepParameterizationBinding>& Binding : BindingsToRemove )
+		{
+			ParameterName.Add( RemoveBinding( Binding ) );
+		}
+
+		return ParameterName;
+	}
+
+	return {};
+}
+
 const UDataprepParameterizationBindings::FBindingToParameterNameMap& UDataprepParameterizationBindings::GetBindingToParameterName() const
 {
 	return BindingToParameterName;
@@ -663,7 +682,7 @@ UObject* UDataprepParameterization::GetDefaultObject()
 	return DefaultParameterisation;
 }
 
-bool UDataprepParameterization::BindObjectProperty(UObject* InObject, const TArray<FDataprepPropertyLink>& InPropertyChain, FName ParameterName)
+bool UDataprepParameterization::BindObjectProperty(UObject* InObject, const TArray<FDataprepPropertyLink>& InPropertyChain, const FName& ParameterName)
 {
 	if ( InObject && InPropertyChain.Num() > 0 && !ParameterName.IsNone() )
 	{
@@ -687,7 +706,6 @@ bool UDataprepParameterization::BindObjectProperty(UObject* InObject, const TArr
 			else
 			{
 				Modify();
-				OnCustomClassAboutToBeUpdated.Broadcast();
 
 				BindingsContainer->Add( Binding, ParameterName );
 			
@@ -725,12 +743,31 @@ void UDataprepParameterization::RemoveBindedObjectProperty(UObject* Object, cons
 	FName ParameterOfRemovedBinding = BindingsContainer->RemoveBinding( Binding );
 	if ( !BindingsContainer->HasBindingsForParameter( ParameterOfRemovedBinding ) )
 	{
-		OnCustomClassAboutToBeUpdated.Broadcast();
 		NameToParameterizationProperty.Remove( ParameterOfRemovedBinding );
 		UpdateClass();
 	}
 }
 
+void UDataprepParameterization::RemoveBindingFromObjects(TArray<UObject*> Objects)
+{
+	Modify();
+	TSet<FName> ParameterPotentiallyRemoved;
+	for ( UObject* Object : Objects )
+	{
+		ParameterPotentiallyRemoved.Append( BindingsContainer->RemoveAllBindingsFromObject( Object ) );
+	}
+
+	for ( const FName& Name : ParameterPotentiallyRemoved )
+	{
+		if ( !BindingsContainer->HasBindingsForParameter( Name ) )
+		{
+			NameToParameterizationProperty.Remove( Name );
+		}
+	}
+
+	UpdateClass();
+}
+
 void UDataprepParameterization::GenerateClass()
 {
 	if ( !CustomContainerClass )
@@ -771,15 +808,10 @@ void UDataprepParameterization::GenerateClass()
 
 void UDataprepParameterization::UpdateClass()
 {
+	OnCustomClassAboutToBeUpdated.Broadcast();
+
 	// Move away the old class
-	if ( CustomContainerClass )
-	{
-		const FString OldClassName = MakeUniqueObjectName( CustomContainerClass->GetOuter(), CustomContainerClass->GetClass(), *FString::Printf(TEXT("%s_REINST"), *CustomContainerClass->GetName()) ).ToString();
-		CustomContainerClass->ClassFlags |= CLASS_NewerVersionExists;
-		CustomContainerClass->SetFlags( RF_NewerVersionExists );
-		CustomContainerClass->ClearFlags( RF_Public | RF_Standalone );
-		CustomContainerClass->Rename( *OldClassName, nullptr, REN_DontCreateRedirectors | REN_DoNotDirty );
-	}
+	PrepareCustomClassForNewClassGeneration();
 
 	UClass* OldClass = CustomContainerClass;
 
@@ -803,7 +835,7 @@ void UDataprepParameterization::LoadParameterization()
 
 		for ( const TPair<TSharedRef<FDataprepParameterizationBinding>, FName>& Binding : BindingsContainer->GetBindingToParameterName() )
 		{
-			const FName BindingName = Binding.Value;
+			const FName& BindingName = Binding.Value;
 
 			UProperty* PropertyFromChain = DataprepParameterization::GetPropertyFromBinding( Binding.Key.Get() );
 			UProperty** PropertyFromParameterizationClass = NameToParameterizationProperty.Find( BindingName );
@@ -912,7 +944,7 @@ void UDataprepParameterization::DoReinstancing(UClass* OldClass, bool bMigrateDa
 	}
 }
 
-UProperty* UDataprepParameterization::AddPropertyToClass(FName ParameterisationPropertyName, UProperty& Property)
+UProperty* UDataprepParameterization::AddPropertyToClass(const FName& ParameterisationPropertyName, UProperty& Property)
 {
 	if ( !NameToParameterizationProperty.Find( ParameterisationPropertyName ) )
 	{
diff --git a/Engine/Plugins/Experimental/DataPrepEditor/Source/DataPrepCore/Private/Parameterization/DataprepParameterization.h b/Engine/Plugins/Experimental/DataPrepEditor/Source/DataPrepCore/Private/Parameterization/DataprepParameterization.h
--- a/Engine/Plugins/Experimental/DataPrepEditor/Source/DataPrepCore/Private/Parameterization/DataprepParameterization.h
+++ b/Engine/Plugins/Experimental/DataPrepEditor/Source/DataPrepCore/Private/Parameterization/DataprepParameterization.h
@@ -91,8 +91,6 @@ struct FDataprepParametrizationBindingSetKeyFuncs : DefaultKeyFuncs<TSharedRef<F
 	}
 };
 
-
-
 /**
  * Encapsulate the unidirectionality necessary for a constant cost of access to the data related to the bindings
  */
@@ -113,12 +111,12 @@ public:
 	/**
 	 * Does the data structure has some bindings for the parameter name
 	 */
-	bool HasBindingsForParameter(FName ParameterName) const;
+	bool HasBindingsForParameter(const FName& ParameterName) const;
 
 	/**
 	 * Add a binding and map it to the parameter
 	 */
-	void Add(const TSharedRef<FDataprepParameterizationBinding>& Binding, FName ParameterName);
+	void Add(const TSharedRef<FDataprepParameterizationBinding>& Binding, const FName& ParameterName);
 
 	/**
 	 * Remove a binding.
@@ -126,6 +124,13 @@ public:
 	 */
 	FName RemoveBinding(const TSharedRef<FDataprepParameterizationBinding>& Binding);
 
+
+	/**
+	 * Remove all the bindings from a object
+	 * @return The name of the parameters that were associated to the binding of the object
+	 */
+	TSet<FName> RemoveAllBindingsFromObject(UObject* Object);
+
 	const FBindingToParameterNameMap& GetBindingToParameterName() const;
 
 	// UObject interface
@@ -179,15 +184,16 @@ public:
 
 	void OnObjectModified(UObject* Object);
 
-
 	UObject* GetDefaultObject();
 
-	bool BindObjectProperty(UObject* Object, const TArray<FDataprepPropertyLink>& PropertyChain, FName Name);
+	bool BindObjectProperty(UObject* Object, const TArray<FDataprepPropertyLink>& PropertyChain, const FName& Name);
 
 	bool IsObjectPropertyBinded(UObject* Object, const TArray<FDataprepPropertyLink>& PropertyChain) const;
 
 	void RemoveBindedObjectProperty(UObject* Object, const TArray<FDataprepPropertyLink>& PropertyChain);
 
+	void RemoveBindingFromObjects(TArray<UObject*> Objects);
+
 private:
 
 	/**
@@ -199,7 +205,7 @@ private:
 	 * Update the Custom Container Class to a newer version
 	 */
 	void UpdateClass();
-	
+
 	/**
 	 * Do the process of regenerating the Custom Container Class and the data of its default object from the serialized data
 	 */
@@ -221,7 +227,7 @@ private:
 	 * Try adding a binded property to the parameterization class
 	 * @return false if the binding is no more valid
 	 */
-	UProperty* AddPropertyToClass(FName ParameterisationPropertyName, UProperty& Property);
+	UProperty* AddPropertyToClass(const FName& ParameterisationPropertyName, UProperty& Property);
 
 	// The containers for the bindings
 	UPROPERTY()
diff --git a/Engine/Plugins/Experimental/DataPrepEditor/Source/DataPrepCore/Public/DataPrepAsset.h b/Engine/Plugins/Experimental/DataPrepEditor/Source/DataPrepCore/Public/DataPrepAsset.h
--- a/Engine/Plugins/Experimental/DataPrepEditor/Source/DataPrepCore/Public/DataPrepAsset.h
+++ b/Engine/Plugins/Experimental/DataPrepEditor/Source/DataPrepCore/Public/DataPrepAsset.h
@@ -91,6 +91,8 @@ public:
 	FOnDataprepBlueprintChange& GetOnBlueprintChanged() { return OnBlueprintChanged; }
 	// end of temp code for nodes development
 
+public:
+	// Functions specific to the parametrization of the dataprep asset
 
 	/**
 	 * Delegate to notify the ui that a dataprep parametrization was modified
@@ -101,7 +103,7 @@ public:
 
 	virtual UObject* GetParameterizationObject() override;
 
-	void BindObjectPropertyToParameterization(UObject* Object, const TArray<struct FDataprepPropertyLink>& InPropertyChain, FName Name);
+	void BindObjectPropertyToParameterization(UObject* Object, const TArray<struct FDataprepPropertyLink>& InPropertyChain,const FName& Name);
 
 	bool IsObjectPropertyBinded(UObject* Object, const TArray<struct FDataprepPropertyLink>& InPropertyChain) const;
 
diff --git a/Engine/Plugins/Experimental/DataPrepEditor/Source/DataPrepCore/Public/DataprepActionAsset.h b/Engine/Plugins/Experimental/DataPrepEditor/Source/DataPrepCore/Public/DataprepActionAsset.h
--- a/Engine/Plugins/Experimental/DataPrepEditor/Source/DataPrepCore/Public/DataprepActionAsset.h
+++ b/Engine/Plugins/Experimental/DataPrepEditor/Source/DataPrepCore/Public/DataprepActionAsset.h
@@ -250,6 +250,11 @@ public:
 	const TCHAR* GetLabel() { return *Label; }
 	void SetLabel( const TCHAR* InLabel ) { Label = InLabel ? InLabel : TEXT(""); }
 
+	/**
+	 * Do the necessary notification so that the dataprep system can react properly to removal of this action
+	 */
+	void NotifyDataprepSystemsOfRemoval();
+
 private:
 
 	void OnClassesRemoved(const TArray<UClass*>& DeletedClasses);
diff --git a/Engine/Plugins/Experimental/DataPrepEditor/Source/DataPrepCore/Public/DataprepCoreUtils.h b/Engine/Plugins/Experimental/DataPrepEditor/Source/DataPrepCore/Public/DataprepCoreUtils.h
--- a/Engine/Plugins/Experimental/DataPrepEditor/Source/DataPrepCore/Public/DataprepCoreUtils.h
+++ b/Engine/Plugins/Experimental/DataPrepEditor/Source/DataPrepCore/Public/DataprepCoreUtils.h
@@ -11,11 +11,21 @@
 #include "UObject/UObjectGlobals.h"
 #include "UObject/Package.h"
 
+class UDataprepAsset;
+
 struct FScopedSlowTask;
 
 class DATAPREPCORE_API FDataprepCoreUtils
 {
 public:
+
+	/**
+	 * Return the dataprep asset that own the object, if the object is part of a dataprep asset
+	 * @return nullptr if the object is not a part of a dataprep asset
+	 */
+	static UDataprepAsset* GetDataprepAssetOfObject(UObject* Object);
+
+	/** Delete the objects and do the manipulation required to safely delete the assets */
 	static void PurgeObjects(TArray<UObject*> Objects);
 
 	/**
diff --git a/Engine/Plugins/Experimental/DataPrepEditor/Source/DataPrepEditor/Private/BlueprintNodes/K2Node_DataprepAction.cpp b/Engine/Plugins/Experimental/DataPrepEditor/Source/DataPrepEditor/Private/BlueprintNodes/K2Node_DataprepAction.cpp
--- a/Engine/Plugins/Experimental/DataPrepEditor/Source/DataPrepEditor/Private/BlueprintNodes/K2Node_DataprepAction.cpp
+++ b/Engine/Plugins/Experimental/DataPrepEditor/Source/DataPrepEditor/Private/BlueprintNodes/K2Node_DataprepAction.cpp
@@ -118,6 +118,9 @@ void UK2Node_DataprepAction::DestroyNode()
 	if ( DataprepActionAsset )
 	{
 		Modify();
+
+		DataprepActionAsset->NotifyDataprepSystemsOfRemoval();
+
 		// Force the transaction system to restore the action
 		DataprepActionAsset = nullptr;
 	}
diff --git a/Engine/Plugins/Experimental/DataPrepEditor/Source/DataPrepEditor/Private/DataPrepEditorModule.cpp b/Engine/Plugins/Experimental/DataPrepEditor/Source/DataPrepEditor/Private/DataPrepEditorModule.cpp
--- a/Engine/Plugins/Experimental/DataPrepEditor/Source/DataPrepEditor/Private/DataPrepEditorModule.cpp
+++ b/Engine/Plugins/Experimental/DataPrepEditor/Source/DataPrepEditor/Private/DataPrepEditorModule.cpp
@@ -7,6 +7,11 @@
 #include "DataprepAssetProducers.h"
 #include "DataPrepEditor.h"
 
+#include "DataprepEditorStyle.h"
+#include "Widgets/DataprepWidgets.h"
+#include "Widgets/SDataprepEditorViewport.h"
+#include "Widgets/SDataprepProducersWidget.h"
+
 #include "Developer/AssetTools/Public/IAssetTools.h"
 #include "Developer/AssetTools/Public/AssetToolsModule.h"
 #include "Misc/PackageName.h"
@@ -14,9 +19,6 @@
 #include "PropertyEditorModule.h"
 #include "UObject/StrongObjectPtr.h"
 
-#include "DataprepEditorStyle.h"
-#include "Widgets/DataprepWidgets.h"
-#include "Widgets/SDataprepProducersWidget.h"
 #include "Widgets/SNullWidget.h"
 
 const FName DataprepEditorAppIdentifier = FName(TEXT("DataprepEditorApp"));
@@ -55,6 +57,8 @@ public:
 		// Register the details customizer
 		FPropertyEditorModule& PropertyModule = FModuleManager::LoadModuleChecked< FPropertyEditorModule >( TEXT("PropertyEditor") );
 		PropertyModule.RegisterCustomClassLayout( UDataprepAssetProducers::StaticClass()->GetFName(), FOnGetDetailCustomizationInstance::CreateStatic( &FDataprepAssetProducersDetails::MakeDetails ) );
+
+		SDataprepEditorViewport::LoadDefaultSettings();
 	}
 
 	virtual void ShutdownModule() override
diff --git a/Engine/Plugins/Experimental/DataPrepEditor/Source/DataPrepEditor/Private/Widgets/SDataprepEditorViewport.cpp b/Engine/Plugins/Experimental/DataPrepEditor/Source/DataPrepEditor/Private/Widgets/SDataprepEditorViewport.cpp
--- a/Engine/Plugins/Experimental/DataPrepEditor/Source/DataPrepEditor/Private/Widgets/SDataprepEditorViewport.cpp
+++ b/Engine/Plugins/Experimental/DataPrepEditor/Source/DataPrepEditor/Private/Widgets/SDataprepEditorViewport.cpp
@@ -8,11 +8,13 @@
 #include "DataprepEditorLogCategory.h"
 
 #include "ActorEditorUtils.h"
+#include "AssetViewerSettings.h"
 #include "Async/ParallelFor.h"
 #include "ComponentReregisterContext.h"
 #include "Editor.h"
 #include "Engine/Engine.h"
 #include "Engine/StaticMesh.h"
+#include "Engine/TextureCube.h"
 #include "Engine/World.h"
 #include "EngineUtils.h"
 #include "GenericPlatform/GenericPlatformTime.h"
@@ -25,11 +27,13 @@
 #include "MaterialShared.h"
 #include "MeshDescription.h"
 #include "MeshDescriptionOperations.h"
+#include "Misc/ConfigCacheIni.h"
 #include "Misc/Paths.h"
 #include "Misc/ScopedSlowTask.h"
 #include "PhysicsEngine/BodySetup.h"
 #include "SceneInterface.h"
 #include "SEditorViewportToolBarMenu.h"
+#include "Settings/LevelEditorViewportSettings.h"
 #include "Slate/SceneViewport.h"
 #include "Widgets/Docking/SDockTab.h"
 #include "Widgets/SBoxPanel.h"
@@ -45,8 +49,10 @@ TWeakObjectPtr<UMaterial> SDataprepEditorViewport::XRayMaterial;
 TWeakObjectPtr<UMaterial> SDataprepEditorViewport::BackFaceMaterial;
 TWeakObjectPtr<UMaterial> SDataprepEditorViewport::PerMeshMaterial;
 TWeakObjectPtr<UMaterial> SDataprepEditorViewport::ReflectionMaterial;
-TWeakObjectPtr<UMaterialInstanceConstant> SDataprepEditorViewport::SelectionMaterial;
+TWeakObjectPtr<UMaterialInstanceConstant> SDataprepEditorViewport::TransparentMaterial;
 TArray<TWeakObjectPtr<UMaterialInstanceConstant>> SDataprepEditorViewport::PerMeshMaterialInstances;
+int32 SDataprepEditorViewport::AssetViewerProfileIndex = INDEX_NONE;
+
 
 const FColor PerMeshColor[] =
 {
@@ -222,6 +228,7 @@ SDataprepEditorViewport::SDataprepEditorViewport()
 	, Extender( MakeShareable( new FExtender ) )
 	, WorldToPreview( nullptr )
 	, RenderingMaterialType( ERenderingMaterialType::OriginalRenderingMaterial )
+	, CurrentSelectionMode( ESelectionModeType::OutlineSelectionMode )
 	, bWireframeRenderingMode( false )
 #ifdef VIEWPORT_EXPERIMENTAL
 	, bShowOrientedBox( false )
@@ -471,6 +478,8 @@ void SDataprepEditorViewport::UpdateMeshes()
 							PreviewMeshComponent->SetMaterial( Index, DisplayMaterialsMap[MaterialInterface].Get() );
 						}
 
+						PreviewMeshComponent->SelectionOverrideDelegate = UPrimitiveComponent::FSelectionOverride::CreateRaw( this, &SDataprepEditorViewport::IsComponentSelected );
+
 						PreviewMeshComponents.Emplace( PreviewMeshComponent );
 
 						MeshComponentsMapping.Add( SceneMeshComponent, PreviewMeshComponent );
@@ -699,6 +708,17 @@ void SDataprepEditorViewport::BindCommands()
 		FCanExecuteAction(),
 		FIsActionChecked::CreateSP( this, &SDataprepEditorViewport::IsRenderingMaterialApplied, ERenderingMaterialType::ReflectionRenderingMaterial ) );
 
+	CommandList->MapAction(
+		Commands.ApplyOutlineSelection,
+		FExecuteAction::CreateSP( this, &SDataprepEditorViewport::SetSelectionMode, ESelectionModeType::OutlineSelectionMode ),
+		FCanExecuteAction(),
+		FIsActionChecked::CreateSP( this, &SDataprepEditorViewport::IsSetSelectionModeApplied, ESelectionModeType::OutlineSelectionMode ) );
+
+	CommandList->MapAction(
+		Commands.ApplyXRaySelection,
+		FExecuteAction::CreateSP( this, &SDataprepEditorViewport::SetSelectionMode, ESelectionModeType::XRaySelectionMode ),
+		FCanExecuteAction(),
+		FIsActionChecked::CreateSP( this, &SDataprepEditorViewport::IsSetSelectionModeApplied, ESelectionModeType::XRaySelectionMode ) );
 
 	CommandList->MapAction(
 		Commands.ApplyWireframeMode,
@@ -770,20 +790,20 @@ void SDataprepEditorViewport::InitializeDefaultMaterials()
 		Materials.Add( PreviewMaterial.Get() );
 	}
 
-	if(!SelectionMaterial.IsValid())
+	if(!TransparentMaterial.IsValid())
 	{
-		SelectionMaterial = TWeakObjectPtr<UMaterialInstanceConstant>( NewObject<UMaterialInstanceConstant>( GetTransientPackage(), NAME_None, EObjectFlags::RF_Transient) );
-		SelectionMaterial->Parent = PreviewMaterial.Get();
+		TransparentMaterial = TWeakObjectPtr<UMaterialInstanceConstant>( NewObject<UMaterialInstanceConstant>( GetTransientPackage(), NAME_None, EObjectFlags::RF_Transient) );
+		TransparentMaterial->Parent = PreviewMaterial.Get();
 
-		SelectionMaterial->BasePropertyOverrides.bOverride_BlendMode = true;
-		SelectionMaterial->BasePropertyOverrides.BlendMode = BLEND_Translucent;
+		TransparentMaterial->BasePropertyOverrides.bOverride_BlendMode = true;
+		TransparentMaterial->BasePropertyOverrides.BlendMode = BLEND_Translucent;
 
-		SelectionMaterial->SetScalarParameterValueEditorOnly( TEXT("Transparency"), 0.75f );
-		SelectionMaterial->SetVectorParameterValueEditorOnly( TEXT( "DiffuseColor" ), FLinearColor::Gray );
+		TransparentMaterial->SetScalarParameterValueEditorOnly( TEXT("Transparency"), 0.75f );
+		TransparentMaterial->SetVectorParameterValueEditorOnly( TEXT( "DiffuseColor" ), FLinearColor::Gray );
 
-		check( SelectionMaterial.IsValid() );
+		check( TransparentMaterial.IsValid() );
 
-		Materials.Add( SelectionMaterial.Get() );
+		Materials.Add( TransparentMaterial.Get() );
 	}
 
 	if(!XRayMaterial.IsValid())
@@ -1049,20 +1069,6 @@ void SDataprepEditorViewport::UpdateSelection()
 		return;
 	}
 
-	// Set all meshes transparent
-	for( const TWeakObjectPtr< UStaticMeshComponent >& PreviewMeshComponent : PreviewMeshComponents )
-	{
-		if(UStaticMeshComponent* MeshComponent = PreviewMeshComponent.Get())
-		{
-			for(UMaterialInterface*& MaterialInterface : MeshComponent->OverrideMaterials)
-			{
-				MaterialInterface = SelectionMaterial.Get();
-			}
-
-			MeshComponent->MarkRenderStateDirty();
-		}
-	}
-
 	// Apply materials. Only selected ones will be affected
 	ApplyRenderingMaterial();
 
@@ -1084,6 +1090,12 @@ void SDataprepEditorViewport::UpdateSelection()
 	SceneViewport->Invalidate();
 }
 
+bool SDataprepEditorViewport::IsComponentSelected(const UPrimitiveComponent* InPrimitiveComponent)
+{
+	UPrimitiveComponent* PrimitiveComponent = const_cast<UPrimitiveComponent*>(InPrimitiveComponent);
+	return SelectedPreviewComponents.Contains( Cast<UCustomStaticMeshComponent>(PrimitiveComponent) ) && CurrentSelectionMode == ESelectionModeType::OutlineSelectionMode;
+}
+
 void SDataprepEditorViewport::SetRenderingMaterial(ERenderingMaterialType InRenderingMaterialType)
 {
 	if(RenderingMaterialType != InRenderingMaterialType)
@@ -1109,6 +1121,16 @@ void SDataprepEditorViewport::ToggleWireframeRenderingMode()
 	SceneViewport->Invalidate();
 }
 
+void SDataprepEditorViewport::SetSelectionMode(ESelectionModeType InSelectionMode)
+{
+	if(CurrentSelectionMode != InSelectionMode)
+	{
+		CurrentSelectionMode = InSelectionMode;
+
+		ApplyRenderingMaterial();
+	}
+}
+
 UMaterialInterface* SDataprepEditorViewport::GetRenderingMaterial(UStaticMeshComponent* PreviewMeshComponent)
 {
 	switch(RenderingMaterialType)
@@ -1173,9 +1195,46 @@ void SDataprepEditorViewport::ApplyRenderingMaterial()
 
 	if(SelectedPreviewComponents.Num() > 0)
 	{
-		for( UStaticMeshComponent* PreviewMeshComponent : SelectedPreviewComponents )
+		switch(CurrentSelectionMode)
 		{
-			ApplyMaterial( PreviewMeshComponent );
+			case ESelectionModeType::XRaySelectionMode:
+			{
+				// Apply transparent material on all mesh components
+				for( const TWeakObjectPtr< UStaticMeshComponent >& PreviewMeshComponentPtr : PreviewMeshComponents )
+				{
+					if(UStaticMeshComponent* PreviewMeshComponent = PreviewMeshComponentPtr.Get())
+					{
+						UStaticMesh* StaticMesh = PreviewMeshComponent->GetStaticMesh();
+
+						for(int32 Index = 0; Index < StaticMesh->StaticMaterials.Num(); ++Index)
+						{
+							PreviewMeshComponent->SetMaterial( Index, TransparentMaterial.Get() );
+						}
+
+						PreviewMeshComponent->MarkRenderStateDirty();
+					}
+				}
+
+				// Apply rendering material only on selected mesh components
+				for( UStaticMeshComponent* PreviewMeshComponent : SelectedPreviewComponents )
+				{
+					ApplyMaterial( PreviewMeshComponent );
+				}
+			}
+			break;
+
+			default:
+			case ESelectionModeType::OutlineSelectionMode:
+			{
+				for( const TWeakObjectPtr< UStaticMeshComponent >& PreviewMeshComponentPtr : PreviewMeshComponents )
+				{
+					if(UStaticMeshComponent* PreviewMeshComponent = PreviewMeshComponentPtr.Get())
+					{
+						ApplyMaterial( PreviewMeshComponent );
+					}
+				}
+			}
+			break;
 		}
 	}
 	else
@@ -1192,6 +1251,68 @@ void SDataprepEditorViewport::ApplyRenderingMaterial()
 	SceneViewport->Invalidate();
 }
 
+void SDataprepEditorViewport::LoadDefaultSettings()
+{
+	// Find index of Dataprep's viewport's settings
+	const TCHAR* DataprepViewportSettingProfileName = TEXT("DataprepViewportSetting");
+
+	UAssetViewerSettings* DefaultSettings = UAssetViewerSettings::Get();
+	check(DefaultSettings);
+
+	for(int32 Index = 0; Index < DefaultSettings->Profiles.Num(); ++Index)
+	{
+		if(DefaultSettings->Profiles[Index].ProfileName == DataprepViewportSettingProfileName)
+		{
+			AssetViewerProfileIndex = Index;
+			break;
+		}
+	}
+
+	// No profile found, create one
+	if(AssetViewerProfileIndex == INDEX_NONE)
+	{
+		FPreviewSceneProfile Profile = DefaultSettings->Profiles[0];
+
+		Profile.bSharedProfile = false;
+		Profile.ProfileName = DataprepViewportSettingProfileName;
+		AssetViewerProfileIndex = DefaultSettings->Profiles.Num();
+
+		DefaultSettings->Profiles.Add( Profile );
+		DefaultSettings->Save();
+	}
+
+	// Update the profile with the settings for the project
+	FPreviewSceneProfile& DataprepViewportSettingProfile = 	DefaultSettings->Profiles[AssetViewerProfileIndex];
+
+	// Read default settings, tessellation and import, for Datasmith file producer
+	const FString DataprepEditorIni = FString::Printf(TEXT("%s%s/%s.ini"), *FPaths::GeneratedConfigDir(), ANSI_TO_TCHAR(FPlatformProperties::PlatformName()), TEXT("DataprepEditor") );
+
+	const TCHAR* ViewportSectionName = TEXT("ViewportSettings");
+	if(GConfig->DoesSectionExist( ViewportSectionName, DataprepEditorIni ))
+	{
+		FString EnvironmentCubeMapPath = GConfig->GetStr( ViewportSectionName, TEXT("EnvironmentCubeMap"), DataprepEditorIni);
+
+		if(EnvironmentCubeMapPath != DataprepViewportSettingProfile.EnvironmentCubeMapPath)
+		{
+			// Check that the Cube map does exist
+			FSoftObjectPath EnvironmentCubeMap( DataprepViewportSettingProfile.EnvironmentCubeMapPath );
+			UObject* LoadedObject = EnvironmentCubeMap.TryLoad();
+
+			while (UObjectRedirector* Redirector = Cast<UObjectRedirector>( LoadedObject ))
+			{
+				LoadedObject = Redirector->DestinationObject;
+			}
+
+			// Good to go, update the profile's related parameters
+			if( Cast<UTextureCube>( LoadedObject ) != nullptr )
+			{
+				DataprepViewportSettingProfile.EnvironmentCubeMapPath = EnvironmentCubeMapPath;
+				DataprepViewportSettingProfile.EnvironmentCubeMap = EnvironmentCubeMap;
+			}
+		}
+	}
+}
+
 //
 // FDataprepEditorViewportClient Class
 //
@@ -1201,12 +1322,17 @@ FDataprepEditorViewportClient::FDataprepEditorViewportClient(const TSharedRef<SE
 	, AdvancedPreviewScene( nullptr )
 	, DataprepEditorViewport( nullptr )
 {
+	EngineShowFlags.SetSelectionOutline( true );
+
 	if (EditorViewportWidget.IsValid())
 	{
 		DataprepEditorViewport = static_cast<SDataprepEditorViewport*>( EditorViewportWidget.Pin().Get() );
 	}
 
 	AdvancedPreviewScene = static_cast<FAdvancedPreviewScene*>(PreviewScene);
+	check(AdvancedPreviewScene);
+
+	AdvancedPreviewScene->SetProfileIndex( SDataprepEditorViewport::AssetViewerProfileIndex );
 }
 
 bool FDataprepEditorViewportClient::InputKey(FViewport * InViewport, int32 ControllerId, FKey Key, EInputEvent Event, float AmountDepressed, bool bGamepad)
@@ -1430,10 +1556,16 @@ TSharedRef<SWidget> SDataprepEditorViewportToolbar::GenerateRenderingMenu() cons
 		MenuBuilder.BeginSection("DataprepEditorViewportRenderingMenu", LOCTEXT("DataprepEditor_RenderingMaterial", "Materials"));
 		MenuBuilder.AddMenuEntry(Commands.ApplyOriginalMaterial);
 		MenuBuilder.AddMenuEntry(Commands.ApplyBackFaceMaterial);
+#ifdef VIEWPORT_EXPERIMENTAL
 		MenuBuilder.AddMenuEntry(Commands.ApplyXRayMaterial);
+#endif
 		MenuBuilder.AddMenuEntry(Commands.ApplyPerMeshMaterial);
 		MenuBuilder.AddMenuEntry(Commands.ApplyReflectionMaterial);
 		MenuBuilder.EndSection();
+		MenuBuilder.BeginSection("DataprepEditorViewportRenderingMenu", LOCTEXT("DataprepEditor_SelectionMode", "Selection"));
+		MenuBuilder.AddMenuEntry(Commands.ApplyOutlineSelection);
+		MenuBuilder.AddMenuEntry(Commands.ApplyXRaySelection);
+		MenuBuilder.EndSection();
 		MenuBuilder.BeginSection("DataprepEditorViewportRenderingMenu", LOCTEXT("DataprepEditor_RenderingMode", "Modes"));
 		MenuBuilder.AddMenuEntry(Commands.ApplyWireframeMode);
 		MenuBuilder.EndSection();
@@ -1496,6 +1628,10 @@ void FDataprepEditorViewportCommands::RegisterCommands()
 	UI_COMMAND(ApplyXRayMaterial, "XRay", "Use XRay material to render meshes.", EUserInterfaceActionType::RadioButton, FInputChord());
 	UI_COMMAND(ApplyPerMeshMaterial, "MultiColored", "Assign a different color for each rendered mesh.", EUserInterfaceActionType::RadioButton, FInputChord());
 	UI_COMMAND(ApplyReflectionMaterial, "ReflectionLines", "Use reflective material to show lines of reflection.", EUserInterfaceActionType::RadioButton, FInputChord());
+	
+	// Selection Mode
+	UI_COMMAND(ApplyOutlineSelection, "Outline", "Outline selected meshes with a colored contour.", EUserInterfaceActionType::RadioButton, FInputChord());
+	UI_COMMAND(ApplyXRaySelection, "XRay", "Use XRay material on non selected meshes.", EUserInterfaceActionType::RadioButton, FInputChord());
 
 	// Rendering Mode
 	UI_COMMAND(ApplyWireframeMode, "Wireframe", "Display all meshes in wireframe.", EUserInterfaceActionType::ToggleButton, FInputChord());
diff --git a/Engine/Plugins/Experimental/DataPrepEditor/Source/DataPrepEditor/Private/Widgets/SDataprepEditorViewport.h b/Engine/Plugins/Experimental/DataPrepEditor/Source/DataPrepEditor/Private/Widgets/SDataprepEditorViewport.h
--- a/Engine/Plugins/Experimental/DataPrepEditor/Source/DataPrepEditor/Private/Widgets/SDataprepEditorViewport.h
+++ b/Engine/Plugins/Experimental/DataPrepEditor/Source/DataPrepEditor/Private/Widgets/SDataprepEditorViewport.h
@@ -86,6 +86,9 @@ public:
 	TSharedPtr<FUICommandInfo> ApplyPerMeshMaterial;
 	TSharedPtr<FUICommandInfo> ApplyReflectionMaterial;
 
+	TSharedPtr<FUICommandInfo> ApplyOutlineSelection;
+	TSharedPtr<FUICommandInfo> ApplyXRaySelection;
+
 	TSharedPtr<FUICommandInfo> ApplyWireframeMode;
 
 #ifdef VIEWPORT_EXPERIMENTAL
@@ -173,6 +176,9 @@ public:
 	/** Change visibility of preview mesh components */
 	void SetActorVisibility(AActor* SceneActor, bool bInVisibility);
 
+	/** Load the settings affecting the viewport, i.e. environment map */
+	static void LoadDefaultSettings();
+
 protected:
 	/** SEditorViewport interface */
 	virtual TSharedRef<FEditorViewportClient> MakeEditorViewportClient() override;
@@ -194,6 +200,13 @@ private:
 		MaxRenderingMaterial
 	};
 
+	enum class ESelectionModeType : uint8
+	{
+		OutlineSelectionMode,
+		XRaySelectionMode,
+		MaxSelectionMode
+	};
+
 	enum class ERenderingMode : uint8
 	{
 		NormalRenderingMode,
@@ -277,6 +290,15 @@ private:
 	bool IsShowOrientedBoxOn() { return bShowOrientedBox; }
 #endif
 
+	/** Set the selection mode */
+	void SetSelectionMode(ESelectionModeType InSelectionMode);
+
+	/** Check if a given selection mode is active */
+	bool IsSetSelectionModeApplied(ESelectionModeType InSelectionMode) { return InSelectionMode == CurrentSelectionMode; }
+
+	/** Returns true if the given primitive component is part of the selection set */
+	bool IsComponentSelected(const UPrimitiveComponent* PrimitiveComponent);
+
 private:
 	/** The scene for this viewport */
 	TSharedPtr< FAdvancedPreviewScene > PreviewScene;
@@ -321,6 +343,8 @@ private:
 
 	ERenderingMaterialType RenderingMaterialType;
 
+	ESelectionModeType CurrentSelectionMode;
+
 	/** Indicates if wireframe mode is on or off */
 	bool bWireframeRenderingMode;
 
@@ -334,7 +358,7 @@ private:
 	static TWeakObjectPtr<UMaterial> BackFaceMaterial;
 
 	/** Transparent material instance used to display non-selected meshes */
-	static TWeakObjectPtr<UMaterialInstanceConstant> SelectionMaterial;
+	static TWeakObjectPtr<UMaterialInstanceConstant> TransparentMaterial;
 
 	/** Fully reflective material used to display surface discontinuity */
 	static TWeakObjectPtr<UMaterial> ReflectionMaterial;
@@ -351,7 +375,10 @@ private:
 	/** Pointer to the vertical box into which the overlay text items are added */
 	TSharedPtr<SVerticalBox> OverlayTextVerticalBox;
 	TSharedPtr<STextBlock> ScreenSizeText;
-	
+
+	/** Index of the profile to use in the preview scene */
+	static int32 AssetViewerProfileIndex;
+
 #ifdef VIEWPORT_EXPERIMENTAL
 	bool bShowOrientedBox;
 #endif
diff --git a/Engine/Plugins/Experimental/DataPrepEditor/Source/DataPrepLibraries/Private/DataprepEditingOperations.cpp b/Engine/Plugins/Experimental/DataPrepEditor/Source/DataPrepLibraries/Private/DataprepEditingOperations.cpp
--- a/Engine/Plugins/Experimental/DataPrepEditor/Source/DataPrepLibraries/Private/DataprepEditingOperations.cpp
+++ b/Engine/Plugins/Experimental/DataPrepEditor/Source/DataPrepLibraries/Private/DataprepEditingOperations.cpp
@@ -212,7 +212,7 @@ void UDataprepMergeActorsOperation::OnExecution_Implementation(const FDataprepCo
 	DataprepEditingOperationTime::FTimeLogger TimeLogger( TEXT("MergeActors"), [&]( FText Text) { this->LogInfo( Text ); });
 #endif
 
-	if(!MergeStaticMeshActors(CurrentWorld, ComponentsToMerge, TEXT("Merged") ))
+	if(!MergeStaticMeshActors(CurrentWorld, ComponentsToMerge, NewActorLabel.IsEmpty() ? TEXT("Merged") : *NewActorLabel ))
 	{
 		return;
 	}
@@ -278,7 +278,7 @@ bool UDataprepMergeActorsOperation::MergeStaticMeshActors(UWorld* World, const T
 	}
 
 	// Add asset to set of assets in Dataprep action working set
-	MergedMesh = Cast<UStaticMesh>( AddAsset( UtilitiesMergedMesh, UStaticMesh::StaticClass(), *(RootName + TEXT("Mesh")) ) );
+	MergedMesh = Cast<UStaticMesh>( AddAsset( UtilitiesMergedMesh, UStaticMesh::StaticClass(), NewActorLabel.IsEmpty() ? TEXT("Merged_Mesh") : *NewActorLabel ) );
 	if (!MergedMesh)
 	{
 		UE_LOG(LogDataprep, Error, TEXT("MergeStaticMeshActors failed. Internal error while creating the merged mesh."));
@@ -288,7 +288,7 @@ bool UDataprepMergeActorsOperation::MergeStaticMeshActors(UWorld* World, const T
 	if(bCreateActor == true)
 	{
 		// Place new mesh in the world
-		MergedActor = Cast<AStaticMeshActor>( CreateActor( AStaticMeshActor::StaticClass(), *(RootName + TEXT("Actor")) ) );
+		MergedActor = Cast<AStaticMeshActor>( CreateActor( AStaticMeshActor::StaticClass(), NewActorLabel.IsEmpty() ? TEXT("Merged_Actor") : *NewActorLabel ) );
 		if (!MergedActor)
 		{
 			UE_LOG(LogDataprep, Error, TEXT("MergeStaticMeshActors failed. Internal error while creating the merged actor."));
