Merging CL #9583756 from //UE4/Dev-Enterprise

DataprepParameterization react better to the parameterized property changes.
Now, the parameterization always has the same value as is bindings. (still missing the implementation for the container modifications (add, clear, e.c.t.))

diff --git a/Engine/Plugins/Experimental/DataPrepEditor/Source/DataPrepCore/Private/DataPrepAsset.cpp b/Engine/Plugins/Experimental/DataPrepEditor/Source/DataPrepCore/Private/DataPrepAsset.cpp
--- a/Engine/Plugins/Experimental/DataPrepEditor/Source/DataPrepCore/Private/DataPrepAsset.cpp
+++ b/Engine/Plugins/Experimental/DataPrepEditor/Source/DataPrepCore/Private/DataPrepAsset.cpp
@@ -11,6 +11,7 @@
 #include "DataprepCoreLogCategory.h"
 #include "DataprepCorePrivateUtils.h"
 #include "DataprepCoreUtils.h"
+#include "DataprepParameterizableObject.h"
 #include "Parameterization/DataprepParameterization.h"
 
 #include "AssetRegistryModule.h"
@@ -193,7 +194,7 @@ UObject* UDataprepAsset::GetParameterizationObject()
 	return Parameterization->GetDefaultObject();
 }
 
-void UDataprepAsset::BindObjectPropertyToParameterization(UObject* Object, const TArray<FDataprepPropertyLink>& InPropertyChain, const FName& Name)
+void UDataprepAsset::BindObjectPropertyToParameterization(UDataprepParameterizableObject* Object, const TArray<struct FDataprepPropertyLink>& InPropertyChain, const FName& Name)
 {
 	bool bPassConditionCheck = false;
 
@@ -215,12 +216,12 @@ void UDataprepAsset::BindObjectPropertyToParameterization(UObject* Object, const
 
 }
 
-bool UDataprepAsset::IsObjectPropertyBinded(UObject* Object, const TArray<struct FDataprepPropertyLink>& InPropertyChain) const
+bool UDataprepAsset::IsObjectPropertyBinded(UDataprepParameterizableObject* Object, const TArray<struct FDataprepPropertyLink>& InPropertyChain) const
 {
 	return Parameterization->IsObjectPropertyBinded( Object, InPropertyChain );
 }
 
-void UDataprepAsset::RemoveObjectPropertyFromParameterization(UObject* Object, const TArray<struct FDataprepPropertyLink>& InPropertyChain)
+void UDataprepAsset::RemoveObjectPropertyFromParameterization(UDataprepParameterizableObject* Object, const TArray<struct FDataprepPropertyLink>& InPropertyChain)
 {
 	Parameterization->RemoveBindedObjectProperty( Object, InPropertyChain );
 }
diff --git a/Engine/Plugins/Experimental/DataPrepEditor/Source/DataPrepCore/Private/DataprepActionAsset.cpp b/Engine/Plugins/Experimental/DataPrepEditor/Source/DataPrepCore/Private/DataprepActionAsset.cpp
--- a/Engine/Plugins/Experimental/DataPrepEditor/Source/DataPrepCore/Private/DataprepActionAsset.cpp
+++ b/Engine/Plugins/Experimental/DataPrepEditor/Source/DataPrepCore/Private/DataprepActionAsset.cpp
@@ -8,6 +8,7 @@
 #include "DataprepCoreLogCategory.h"
 #include "DataprepCorePrivateUtils.h"
 #include "DataprepCoreUtils.h"
+#include "DataprepParameterizableObject.h"
 #include "IDataprepLogger.h"
 #include "IDataprepProgressReporter.h"
 #include "Parameterization/DataprepParameterization.h"
@@ -258,7 +259,17 @@ bool UDataprepActionAsset::RemoveStep(int32 Index)
 			{
 				TArray< UObject* > Objects;
 				GetObjectsWithOuter( Steps[Index], Objects );
-				Parameterization->RemoveBindingFromObjects( Objects );
+				TArray< UDataprepParameterizableObject* > ParameterizableObjects;
+				ParameterizableObjects.Reserve( Objects.Num() );
+				for ( UObject* Object : Objects )
+				{
+					if ( Object->IsA<UDataprepParameterizableObject>() )
+					{
+						ParameterizableObjects.Add( static_cast<UDataprepParameterizableObject*>( Object ) );
+					}
+				}
+
+				Parameterization->RemoveBindingFromObjects( ParameterizableObjects );
 			}
 		}
 
@@ -284,11 +295,20 @@ void UDataprepActionAsset::NotifyDataprepSystemsOfRemoval()
 		if ( UDataprepParameterization * Parameterization = DataprepAsset->GetDataprepParameterization() )
 		{
 			TArray< UObject* > Objects;
+			TArray< UDataprepParameterizableObject* > ParameterizableObjects;
 			for ( UDataprepActionStep* Step : Steps )
 			{
 				GetObjectsWithOuter( Step, Objects );
+				ParameterizableObjects.Reserve( Objects.Num() );
+				for ( UObject* Object : Objects )
+				{
+					if ( Object->IsA<UDataprepParameterizableObject>() )
+					{
+						ParameterizableObjects.Add( static_cast<UDataprepParameterizableObject*>( Object ) );
+					}
+				}
 			}
-			Parameterization->RemoveBindingFromObjects( Objects );
+			Parameterization->RemoveBindingFromObjects( ParameterizableObjects );
 		}
 	}
 }
diff --git a/Engine/Plugins/Experimental/DataPrepEditor/Source/DataPrepCore/Private/DataprepParameterizableObject.cpp b/Engine/Plugins/Experimental/DataPrepEditor/Source/DataPrepCore/Private/DataprepParameterizableObject.cpp
new file mode 100644
--- /dev/null
+++ b/Engine/Plugins/Experimental/DataPrepEditor/Source/DataPrepCore/Private/DataprepParameterizableObject.cpp
@@ -0,0 +1,39 @@
+// Copyright 1998-2019 Epic Games, Inc. All Rights Reserved.
+
+#include "DataprepParameterizableObject.h"
+
+#include "DataPrepAsset.h"
+#include "DataprepCoreUtils.h"
+#include "Parameterization/DataprepParameterization.h"
+#include "Parameterization/DataprepParameterizationUtils.h"
+
+#include "UObject/UnrealType.h"
+
+void UDataprepParameterizableObject::PostEditChangeChainProperty(FPropertyChangedChainEvent& PropertyChangedEvent)
+{
+	Super::PostEditChangeChainProperty( PropertyChangedEvent );
+
+	if ( !bool( PropertyChangedEvent.ChangeType & ( EPropertyChangeType::Interactive | EPropertyChangeType::Redirected ) ) )
+	{
+		UDataprepAsset* DataprepAsset = FDataprepCoreUtils::GetDataprepAssetOfObject( this );
+		UDataprepParameterization * Parameterization = nullptr;
+		if ( DataprepAsset )
+		{
+			Parameterization = DataprepAsset->GetDataprepParameterization();
+		}
+
+		if ( !Parameterization )
+		{
+			UClass* Class = GetClass();
+			FString PathToDataprepParametrization = Class->GetMetaData( UDataprepParameterization::MetadataClassGeneratorName );
+			Parameterization = FindObject<UDataprepParameterization>( nullptr, *PathToDataprepParametrization );
+			
+		}
+
+		if ( Parameterization )
+		{
+			TArray<FDataprepPropertyLink> PropertyChain = FDataprepParameterizationUtils::MakePropertyChain( PropertyChangedEvent );
+			Parameterization->OnObjectPostEdit( this, PropertyChain, PropertyChangedEvent.ChangeType );
+		}
+	}
+}
diff --git a/Engine/Plugins/Experimental/DataPrepEditor/Source/DataPrepCore/Private/Parameterization/DataprepParameterization.cpp b/Engine/Plugins/Experimental/DataPrepEditor/Source/DataPrepCore/Private/Parameterization/DataprepParameterization.cpp
--- a/Engine/Plugins/Experimental/DataPrepEditor/Source/DataPrepCore/Private/Parameterization/DataprepParameterization.cpp
+++ b/Engine/Plugins/Experimental/DataPrepEditor/Source/DataPrepCore/Private/Parameterization/DataprepParameterization.cpp
@@ -3,6 +3,7 @@
 #include "Parameterization/DataprepParameterization.h"
 
 #include "DataPrepAsset.h"
+#include "DataprepParameterizableObject.h"
 #include "DataprepParameterizationArchive.h"
 
 #include "CoreGlobals.h"
@@ -348,6 +349,21 @@ bool UDataprepParameterizationBindings::ContainsBinding(const TSharedRef<FDatapr
 	return BindingToParameterName.Contains( Binding );
 }
 
+FName UDataprepParameterizationBindings::GetParameterNameForBinding(const TSharedRef<FDataprepParameterizationBinding>& Binding) const
+{
+	return BindingToParameterName.FindRef( Binding );
+}
+
+const UDataprepParameterizationBindings::FSetOfBinding* UDataprepParameterizationBindings::GetBindingsFromObject(const UDataprepParameterizableObject* Object) const
+{
+	return ObjectToBindings.Find( Object );
+}
+
+const UDataprepParameterizationBindings::FSetOfBinding* UDataprepParameterizationBindings::GetBindingsFromParameter(const FName& ParameterName) const
+{
+	return NameToBindings.Find( ParameterName );
+}
+
 bool UDataprepParameterizationBindings::HasBindingsForParameter(const FName& ParameterName) const
 {
 	return NameToBindings.Contains( ParameterName );
@@ -400,7 +416,7 @@ FName UDataprepParameterizationBindings::RemoveBinding(const TSharedRef<FDatapre
 
 		// 3) Remove from the map object to bindings
 		{
-			UObject* Object = Binding->ObjectBinded;
+			UDataprepParameterizableObject* Object = Binding->ObjectBinded;
 			uint32 ObjectHash = GetTypeHash( Object );
 			if ( FSetOfBinding* BindingsMappedToObject = ObjectToBindings.FindByHash( ObjectHash, Object ) )
 			{
@@ -421,7 +437,7 @@ FName UDataprepParameterizationBindings::RemoveBinding(const TSharedRef<FDatapre
 	return NAME_None;
 }
 
-TSet<FName> UDataprepParameterizationBindings::RemoveAllBindingsFromObject(UObject* Object)
+TSet<FName> UDataprepParameterizationBindings::RemoveAllBindingsFromObject(UDataprepParameterizableObject* Object)
 {
 	if ( FSetOfBinding* Bindings = ObjectToBindings.Find( Object ) )
 	{
@@ -464,7 +480,7 @@ void UDataprepParameterizationBindings::AddReferencedObjects(UObject* InThis, FR
 {
 	UDataprepParameterizationBindings* Bindings = CastChecked<UDataprepParameterizationBindings>( InThis );
 
-	for ( TPair<UObject*, FSetOfBinding>& Pair : Bindings->ObjectToBindings )
+	for ( TPair<UDataprepParameterizableObject*, FSetOfBinding>& Pair : Bindings->ObjectToBindings )
 	{
 		Collector.AddReferencedObject( Pair.Key );
 
@@ -498,9 +514,9 @@ void UDataprepParameterizationBindings::Save(FArchive& Ar)
 	{
 		int32 NumberOfObjectBinded = ObjectToBindings.Num();
 		Ar << NumberOfObjectBinded;
-		for ( const TPair<UObject*, FSetOfBinding>& PairObjectAndBindings : ObjectToBindings )
+		for ( const TPair<UDataprepParameterizableObject*, FSetOfBinding>& PairObjectAndBindings : ObjectToBindings )
 		{
-			UObject* Object = PairObjectAndBindings.Key;
+			UDataprepParameterizableObject* Object = PairObjectAndBindings.Key;
 			Ar << Object;
 			int32 NumberOfBinding = PairObjectAndBindings.Value.Num();
 			Ar << NumberOfBinding;
@@ -552,7 +568,7 @@ void UDataprepParameterizationBindings::Load(FArchive& Ar)
 
 		for ( int32 Index = 0; Index < NumOfObjectBinded; Index++ )
 		{
-			UObject* Object = nullptr;
+			UDataprepParameterizableObject* Object = nullptr;
 			Ar << Object;
 
 			FSetOfBinding& Bindings = ObjectToBindings.FindOrAdd( Object );
@@ -682,17 +698,17 @@ UObject* UDataprepParameterization::GetDefaultObject()
 	return DefaultParameterisation;
 }
 
-bool UDataprepParameterization::BindObjectProperty(UObject* InObject, const TArray<FDataprepPropertyLink>& InPropertyChain, const FName& ParameterName)
+bool UDataprepParameterization::BindObjectProperty(UDataprepParameterizableObject* Object, const TArray<FDataprepPropertyLink>& PropertyChain, const FName& Name)
 {
-	if ( InObject && InPropertyChain.Num() > 0 && !ParameterName.IsNone() )
+	if ( Object && PropertyChain.Num() > 0 && !Name.IsNone() )
 	{
-		TSharedRef<FDataprepParameterizationBinding> Binding = MakeShared<FDataprepParameterizationBinding>( InObject, InPropertyChain );
+		TSharedRef<FDataprepParameterizationBinding> Binding = MakeShared<FDataprepParameterizationBinding>( Object, PropertyChain );
 		void* AddressOfTheValueFromBinding;
 
 		// We expect the chain to have a valid chain of cached property before inserting the binding
-		if ( DataprepParameterization::GetDeepestLevelOfValidCache( Binding.Get(), AddressOfTheValueFromBinding) == InPropertyChain.Num() - 1 )
+		if ( DataprepParameterization::GetDeepestLevelOfValidCache( Binding.Get(), AddressOfTheValueFromBinding) == PropertyChain.Num() - 1 )
 		{
-			if ( UProperty** PropertyPtr = NameToParameterizationProperty.Find( ParameterName ) )
+			if ( UProperty** PropertyPtr = NameToParameterizationProperty.Find( Name ) )
 			{
 				UProperty* Property = *PropertyPtr;
 
@@ -700,27 +716,23 @@ bool UDataprepParameterization::BindObjectProperty(UObject* InObject, const TArr
 				if ( Property->GetClass() == Binding->ValueType )
 				{
 					Modify();
-					BindingsContainer->Add( Binding, ParameterName );
+					BindingsContainer->Add( Binding, Name );
+					UpdateParameterizationFromBinding( Binding );
 				}
 			}
 			else
 			{
 				Modify();
 
-				BindingsContainer->Add( Binding, ParameterName );
+				BindingsContainer->Add( Binding, Name );
 			
-				UProperty* SourceProperty = InPropertyChain.Last().CachedProperty.Get();
+				UProperty* SourceProperty = PropertyChain.Last().CachedProperty.Get();
 
 				// The validation we did with GetDeepestLevelOfValidCache ensure us that the property ptr is valid
-				UProperty* NewProperty = AddPropertyToClass( ParameterName, *SourceProperty );
+				UProperty* NewProperty = AddPropertyToClass( Name, *SourceProperty );
 				UpdateClass();
 
-				// Update the value of the new parameter to match the source
-				UObject* Parameterization = CustomContainerClass->GetDefaultObject();
-				UProperty* ParameterizationProperty = FindObjectFast<UProperty>( CustomContainerClass, ParameterName );
-				check( ParameterizationProperty );
-				void* ParameterizationValueAddress = DataprepParameterization::GetAddressOf( *ParameterizationProperty, Parameterization, INDEX_NONE );
-				DataprepParameterization::CopyCompleteValue( *ParameterizationProperty, ParameterizationValueAddress, *SourceProperty, AddressOfTheValueFromBinding );
+				UpdateParameterizationFromBinding( Binding );
 			}
 		}
 	}
@@ -728,13 +740,13 @@ bool UDataprepParameterization::BindObjectProperty(UObject* InObject, const TArr
 	return false;
 }
 
-bool UDataprepParameterization::IsObjectPropertyBinded(UObject* Object, const TArray<FDataprepPropertyLink>& PropertyChain) const
+bool UDataprepParameterization::IsObjectPropertyBinded(UDataprepParameterizableObject* Object, const TArray<FDataprepPropertyLink>& PropertyChain) const
 {
 	TSharedRef<FDataprepParameterizationBinding> Binding= MakeShared<FDataprepParameterizationBinding>( Object, PropertyChain );
 	return BindingsContainer->ContainsBinding( Binding );
 }
 
-void UDataprepParameterization::RemoveBindedObjectProperty(UObject* Object, const TArray<FDataprepPropertyLink>& PropertyChain)
+void UDataprepParameterization::RemoveBindedObjectProperty(UDataprepParameterizableObject* Object, const TArray<FDataprepPropertyLink>& PropertyChain)
 {
 	TSharedRef<FDataprepParameterizationBinding> Binding = MakeShared<FDataprepParameterizationBinding>( Object, PropertyChain );
 
@@ -748,11 +760,11 @@ void UDataprepParameterization::RemoveBindedObjectProperty(UObject* Object, cons
 	}
 }
 
-void UDataprepParameterization::RemoveBindingFromObjects(TArray<UObject*> Objects)
+void UDataprepParameterization::RemoveBindingFromObjects(TArray<UDataprepParameterizableObject*> Objects)
 {
 	Modify();
 	TSet<FName> ParameterPotentiallyRemoved;
-	for ( UObject* Object : Objects )
+	for ( UDataprepParameterizableObject* Object : Objects )
 	{
 		ParameterPotentiallyRemoved.Append( BindingsContainer->RemoveAllBindingsFromObject( Object ) );
 	}
@@ -768,12 +780,55 @@ void UDataprepParameterization::RemoveBindingFromObjects(TArray<UObject*> Object
 	UpdateClass();
 }
 
+void UDataprepParameterization::UpdateParameterizationFromBinding(const TSharedRef<FDataprepParameterizationBinding>& Binding)
+{
+	FName ParameterModified = BindingsContainer->GetParameterNameForBinding(Binding);
+	UProperty* ParametrizationProperty = NameToParameterizationProperty.FindRef(ParameterModified);
+	if (ParametrizationProperty)
+	{
+		void* AddressOfObjectValue = nullptr;
+		if (UProperty * ObjectProperty = DataprepParameterization::GetPropertyFromBinding(Binding.Get(), AddressOfObjectValue))
+		{
+			Modify();
+			void* AddressOfParameterizationValue = DataprepParameterization::GetAddressOf(*ParametrizationProperty, DefaultParameterisation, INDEX_NONE);
+			DataprepParameterization::CopyCompleteValue(*ParametrizationProperty, AddressOfParameterizationValue, *ObjectProperty, AddressOfObjectValue);
+
+			// Post edit the default parameterization
+			FEditPropertyChain EditChain;
+			EditChain.AddHead(ParametrizationProperty);
+			EditChain.SetActivePropertyNode(ParametrizationProperty);
+			FPropertyChangedEvent EditPropertyChangeEvent(ParametrizationProperty, EPropertyChangeType::ValueSet);
+			FPropertyChangedChainEvent EditChangeChainEvent(EditChain, EditPropertyChangeEvent);
+			DefaultParameterisation->PostEditChangeChainProperty(EditChangeChainEvent);
+		}
+	}
+}
+
+void UDataprepParameterization::OnObjectPostEdit(UDataprepParameterizableObject* Object, const TArray<FDataprepPropertyLink>& PropertyChain, EPropertyChangeType::Type ChangeType)
+{
+	if ( Object && PropertyChain.Num() > 0 )
+	{
+		if ( Object == DefaultParameterisation )
+		{
+			PushParametrizationValueToBindings( PropertyChain[0].PropertyName );
+		}
+		else if ( bool( ChangeType & EPropertyChangeType::ValueSet ) )
+		{
+			TSharedRef<FDataprepParameterizationBinding> Binding = MakeShared<FDataprepParameterizationBinding>( Object, PropertyChain );
+			UpdateParameterizationFromBinding( Binding );
+		}
+		else if ( const UDataprepParameterizationBindings::FSetOfBinding* Bindings = BindingsContainer->GetBindingsFromObject( Object ) )
+		{
+			// todo for container ChangeType
+		}
+	}
+}
+
 void UDataprepParameterization::GenerateClass()
 {
 	if ( !CustomContainerClass )
 	{
-		CustomContainerClass = NewObject<UClass>( GetOutermost(), FName( TEXT("Parameterization") ), RF_Transient );
-		CustomContainerClass->SetSuperStruct( UObject::StaticClass() );
+		CreateClassObject();
 
 		// Make the properties appear in a alphabetically order (for that we must add the properties to the class in the reverse order)
 		NameToParameterizationProperty.KeySort([](const FName& First, const FName& Second)
@@ -827,8 +882,7 @@ void UDataprepParameterization::LoadParameterization()
 {
 	if ( !CustomContainerClass )
 	{ 
-		CustomContainerClass = NewObject<UClass>( GetOutermost(), FName(TEXT("Parameterization")), RF_Transient );
-		CustomContainerClass->SetSuperStruct( UObject::StaticClass() );
+		CreateClassObject();
 
 		TSet<TSharedRef<FDataprepParameterizationBinding>> BindingToRemove;
 		NameToParameterizationProperty.Empty( NameToParameterizationProperty.Num() );
@@ -888,6 +942,15 @@ void UDataprepParameterization::PrepareCustomClassForNewClassGeneration()
 	}
 }
 
+void UDataprepParameterization::CreateClassObject()
+{
+	check( !CustomContainerClass );
+
+	CustomContainerClass = NewObject<UClass>(GetOutermost(), FName(TEXT("Parameterization")), RF_Transient);
+	CustomContainerClass->SetSuperStruct( UDataprepParameterizableObject::StaticClass() );
+	CustomContainerClass->SetMetaData( MetadataClassGeneratorName, *GetPathName() );
+}
+
 void UDataprepParameterization::DoReinstancing(UClass* OldClass, bool bMigrateData)
 {
 	if ( OldClass && CustomContainerClass )
@@ -966,6 +1029,45 @@ UProperty* UDataprepParameterization::AddPropertyToClass(const FName& Parameteri
 	return nullptr;
 }
 
+void UDataprepParameterization::PushParametrizationValueToBindings(FName ParameterName)
+{
+	if ( UProperty* ParameterizationPropeterty = NameToParameterizationProperty.FindRef( ParameterName ) )
+	{
+		void* AddressOfParameterValue = DataprepParameterization::GetAddressOf( *ParameterizationPropeterty, DefaultParameterisation, INDEX_NONE );
+
+		if ( AddressOfParameterValue )
+		{
+			if ( const UDataprepParameterizationBindings::FSetOfBinding* Bindings = BindingsContainer->GetBindingsFromParameter( ParameterName ) )
+			{
+				TSet<TSharedRef<FDataprepParameterizationBinding>> BindingToRemove;
+
+				for ( const TSharedRef<FDataprepParameterizationBinding>& Binding : *Bindings )
+				{
+					void* AddressOfBindingValue = nullptr;
+					if ( UProperty* BindingProperty = DataprepParameterization::GetPropertyFromBinding( Binding.Get(), AddressOfBindingValue ) )
+					{
+						Binding->ObjectBinded->Modify();
+						DataprepParameterization::CopyCompleteValue( *BindingProperty, AddressOfBindingValue, *ParameterizationPropeterty, AddressOfParameterValue );
+					}
+					else
+					{
+						BindingToRemove.Add( Binding );
+					}
+				}
+
+				// Remove the invalid bindings
+				for ( const TSharedRef<FDataprepParameterizationBinding>& Binding : BindingToRemove )
+				{
+					// Naive slow implementation for now
+					RemoveBindedObjectProperty( Binding->ObjectBinded, Binding->PropertyChain );
+				}
+			}
+		}
+	}
+}
+
+const FName UDataprepParameterization::MetadataClassGeneratorName( TEXT("DataprepCustomParameterizationGenerator") );
+
 
 UDataprepParameterizationInstance::UDataprepParameterizationInstance()
 {
@@ -1020,9 +1122,9 @@ void UDataprepParameterizationInstance::ApplyParameterization(const TMap<UObject
 	for ( const TPair<TSharedRef<FDataprepParameterizationBinding>, FName>& BindingPair : SourceParameterization->BindingsContainer->GetBindingToParameterName() )
 	{
 		const TSharedRef<FDataprepParameterizationBinding>& Binding = BindingPair.Key;
-		if ( UObject* Object = SourceToCopy.FindRef( Binding->ObjectBinded ) )
+		if ( UDataprepParameterizableObject* Object = Cast<UDataprepParameterizableObject>( SourceToCopy.FindRef( Binding->ObjectBinded ) ) )
 		{
-			TGuardValue<UObject*> GuardObjectBinded( Binding->ObjectBinded, Object );
+			TGuardValue<UDataprepParameterizableObject*> GuardObjectBinded( Binding->ObjectBinded, Object );
 
 			void* DestinationAddress = nullptr;
 			if ( UProperty* DestinationProperty = DataprepParameterization::GetPropertyFromBinding( Binding.Get(), DestinationAddress ) )
diff --git a/Engine/Plugins/Experimental/DataPrepEditor/Source/DataPrepCore/Private/Parameterization/DataprepParameterization.h b/Engine/Plugins/Experimental/DataPrepEditor/Source/DataPrepCore/Private/Parameterization/DataprepParameterization.h
--- a/Engine/Plugins/Experimental/DataPrepEditor/Source/DataPrepCore/Private/Parameterization/DataprepParameterization.h
+++ b/Engine/Plugins/Experimental/DataPrepEditor/Source/DataPrepCore/Private/Parameterization/DataprepParameterization.h
@@ -15,6 +15,7 @@
 
 #include "DataprepParameterization.generated.h"
 
+class UDataprepParameterizableObject;
 class UProperty;
 
 USTRUCT()
@@ -28,7 +29,7 @@ struct FDataprepParameterizationBinding
 		, ValueType( nullptr )
 	{}
 
-	FDataprepParameterizationBinding(UObject* InObjectBinded, TArray<FDataprepPropertyLink> InPropertyChain)
+	FDataprepParameterizationBinding(UDataprepParameterizableObject* InObjectBinded, TArray<FDataprepPropertyLink> InPropertyChain)
 		: ObjectBinded( InObjectBinded )
 		, PropertyChain( MoveTemp( InPropertyChain ) )
 		, ValueType( PropertyChain.Last().CachedProperty.Get() ? PropertyChain.Last().CachedProperty.Get()->GetClass() : nullptr )
@@ -42,7 +43,7 @@ struct FDataprepParameterizationBinding
 	bool operator==(const FDataprepParameterizationBinding& Other) const;
 
 	UPROPERTY()
-	UObject* ObjectBinded;
+	UDataprepParameterizableObject* ObjectBinded;
 
 	UPROPERTY()
 	TArray<FDataprepPropertyLink> PropertyChain;
@@ -102,12 +103,29 @@ public:
 	GENERATED_BODY()
 
 	using FBindingToParameterNameMap = TMap<TSharedRef<FDataprepParameterizationBinding>, FName, FDefaultSetAllocator, FDataprepParametrizationBindingMapKeyFuncs>;
+	using FSetOfBinding = TSet<TSharedRef<FDataprepParameterizationBinding>, FDataprepParametrizationBindingSetKeyFuncs>;
 
 	/**
 	 * Does the data structure contains this binding
 	 */
 	bool ContainsBinding(const TSharedRef<FDataprepParameterizationBinding>& Binding) const;
 
+	/**
+	 * Return the name of the parameter for a binding
+	 */
+	FName GetParameterNameForBinding(const TSharedRef<FDataprepParameterizationBinding>& Binding) const;
+
+	/**
+	 * Return a valid ptr if the object had some bindings
+	 */
+	const FSetOfBinding* GetBindingsFromObject(const UDataprepParameterizableObject* Object) const;
+
+	/**
+	 * Get the bindings from a parameter
+	 * Return nullptr if the parameter doesn't exist
+	 */
+	const FSetOfBinding* GetBindingsFromParameter(const FName& ParameterName) const;
+
 	/**
 	 * Does the data structure has some bindings for the parameter name
 	 */
@@ -124,12 +142,11 @@ public:
 	 */
 	FName RemoveBinding(const TSharedRef<FDataprepParameterizationBinding>& Binding);
 
-
 	/**
 	 * Remove all the bindings from a object
 	 * @return The name of the parameters that were associated to the binding of the object
 	 */
-	TSet<FName> RemoveAllBindingsFromObject(UObject* Object);
+	TSet<FName> RemoveAllBindingsFromObject(UDataprepParameterizableObject* Object);
 
 	const FBindingToParameterNameMap& GetBindingToParameterName() const;
 
@@ -149,17 +166,14 @@ private:
 	 */
 	void Load(FArchive& Ar);
 
-
 	/** Core storage also track a binding to it's parameter name */
 	FBindingToParameterNameMap BindingToParameterName;
 
-	using FSetOfBinding = TSet<TSharedRef<FDataprepParameterizationBinding>, FDataprepParametrizationBindingSetKeyFuncs>;
-
 	/** Track the name usage for parameters */
 	TMap<FName, FSetOfBinding> NameToBindings;
 
 	/** Track which binding a object has */
-	TMap<UObject*, FSetOfBinding> ObjectToBindings;
+	TMap<UDataprepParameterizableObject*, FSetOfBinding> ObjectToBindings;
 };
 
 
@@ -186,13 +200,15 @@ public:
 
 	UObject* GetDefaultObject();
 
-	bool BindObjectProperty(UObject* Object, const TArray<FDataprepPropertyLink>& PropertyChain, const FName& Name);
+	bool BindObjectProperty(UDataprepParameterizableObject* Object, const TArray<FDataprepPropertyLink>& PropertyChain, const FName& Name);
+
+	bool IsObjectPropertyBinded(UDataprepParameterizableObject* Object, const TArray<FDataprepPropertyLink>& PropertyChain) const;
 
-	bool IsObjectPropertyBinded(UObject* Object, const TArray<FDataprepPropertyLink>& PropertyChain) const;
+	void RemoveBindedObjectProperty(UDataprepParameterizableObject* Object, const TArray<FDataprepPropertyLink>& PropertyChain);
 
-	void RemoveBindedObjectProperty(UObject* Object, const TArray<FDataprepPropertyLink>& PropertyChain);
+	void RemoveBindingFromObjects(TArray<UDataprepParameterizableObject*> Objects);
 
-	void RemoveBindingFromObjects(TArray<UObject*> Objects);
+	void OnObjectPostEdit(UDataprepParameterizableObject* Object, const TArray<FDataprepPropertyLink>& PropertyChain, EPropertyChangeType::Type ChangeType);
 
 private:
 
@@ -216,6 +232,11 @@ private:
 	 */
 	void PrepareCustomClassForNewClassGeneration();
 
+	/**
+	 * Do the actual creation of the class object
+	 */
+	void CreateClassObject();
+
 	/**
 	 * Do reinstancing of the objects created from the Custom Container Class
 	 * @param OldClass The previous Custom Constainer Class
@@ -229,6 +250,22 @@ private:
 	 */
 	UProperty* AddPropertyToClass(const FName& ParameterisationPropertyName, UProperty& Property);
 
+	/**
+	 * Get a new value for the parameterization from it's associated binding
+	 */
+	void UpdateParameterizationFromBinding(const TSharedRef<FDataprepParameterizationBinding>& Binding);
+
+	/**
+	 * Push the value of the parametrization to the bindings
+	 */
+	void PushParametrizationValueToBindings(FName ParameterName);
+
+public:
+
+	static const FName MetadataClassGeneratorName;
+
+private:
+
 	// The containers for the bindings
 	UPROPERTY()
 	UDataprepParameterizationBindings* BindingsContainer;
diff --git a/Engine/Plugins/Experimental/DataPrepEditor/Source/DataPrepCore/Private/Parameterization/DataprepParameterizationUtils.cpp b/Engine/Plugins/Experimental/DataPrepEditor/Source/DataPrepCore/Private/Parameterization/DataprepParameterizationUtils.cpp
--- a/Engine/Plugins/Experimental/DataPrepEditor/Source/DataPrepCore/Private/Parameterization/DataprepParameterizationUtils.cpp
+++ b/Engine/Plugins/Experimental/DataPrepEditor/Source/DataPrepCore/Private/Parameterization/DataprepParameterizationUtils.cpp
@@ -4,6 +4,7 @@
 
 #include "DataPrepAsset.h"
 #include "DataPrepOperation.h"
+#include "DataprepParameterizableObject.h"
 #include "SelectionSystem/DataprepFetcher.h"
 #include "SelectionSystem/DataprepFilter.h"
 
@@ -34,9 +35,7 @@ namespace DataprepParameterizationUtils
 
 	bool IsASupportedClassForParameterization(UClass* Class)
 	{
-		return Class == UDataprepOperation::StaticClass()
-			|| Class == UDataprepFilter::StaticClass()
-			|| Class == UDataprepFetcher::StaticClass();
+		return Class->IsChildOf<UDataprepParameterizableObject>();
 	}
 
 }
@@ -103,6 +102,40 @@ TArray<FDataprepPropertyLink> FDataprepParameterizationUtils::MakePropertyChain(
 	return PropertyChain;
 }
 
+TArray<FDataprepPropertyLink> FDataprepParameterizationUtils::MakePropertyChain(FPropertyChangedChainEvent& PropertyChangedEvent)
+{
+	const FEditPropertyChain& EditPropertyChain = PropertyChangedEvent.PropertyChain;
+
+	TArray<FDataprepPropertyLink> DataprepPropertyChain;
+	DataprepPropertyChain.Reserve( EditPropertyChain.Num() + 1 );
+
+	const TDoubleLinkedList<UProperty*>::TDoubleLinkedListNode* CurrentNode = EditPropertyChain.GetHead();
+
+	UProperty* Property = nullptr;
+	while( CurrentNode )
+	{
+		Property = CurrentNode->GetValue();
+		if ( Property )
+		{ 
+			int32 ContainerIndex = PropertyChangedEvent.GetArrayIndex( Property->GetName() );
+			DataprepPropertyChain.Emplace( Property, Property->GetFName(), ContainerIndex);
+		}
+		else
+		{
+			return {};
+		}
+
+		CurrentNode = CurrentNode->GetNextNode();
+	}
+
+	if ( Property != PropertyChangedEvent.Property )
+	{
+		DataprepPropertyChain.Emplace( PropertyChangedEvent.Property, PropertyChangedEvent.Property->GetFName(), INDEX_NONE );
+	}
+
+	return DataprepPropertyChain;
+}
+
 FDataprepParameterizationContext FDataprepParameterizationUtils::CreateContext(TSharedPtr<IPropertyHandle> PropertyHandle, const FDataprepParameterizationContext& ParameterisationContext)
 {
 	FDataprepParameterizationContext NewContext;
diff --git a/Engine/Plugins/Experimental/DataPrepEditor/Source/DataPrepCore/Public/DataPrepAsset.h b/Engine/Plugins/Experimental/DataPrepEditor/Source/DataPrepCore/Public/DataPrepAsset.h
--- a/Engine/Plugins/Experimental/DataPrepEditor/Source/DataPrepCore/Public/DataPrepAsset.h
+++ b/Engine/Plugins/Experimental/DataPrepEditor/Source/DataPrepCore/Public/DataPrepAsset.h
@@ -12,6 +12,7 @@
 
 class UBlueprint;
 class UDataprepActionAsset;
+class UDataprepParameterizableObject;
 class UDataprepParameterization;
 class UDataprepProducers;
 class UEdGraphNode;
@@ -103,11 +104,11 @@ public:
 
 	virtual UObject* GetParameterizationObject() override;
 
-	void BindObjectPropertyToParameterization(UObject* Object, const TArray<struct FDataprepPropertyLink>& InPropertyChain,const FName& Name);
+	void BindObjectPropertyToParameterization(UDataprepParameterizableObject* Object, const TArray<struct FDataprepPropertyLink>& InPropertyChain,const FName& Name);
 
-	bool IsObjectPropertyBinded(UObject* Object, const TArray<struct FDataprepPropertyLink>& InPropertyChain) const;
+	bool IsObjectPropertyBinded(UDataprepParameterizableObject* Object, const TArray<struct FDataprepPropertyLink>& InPropertyChain) const;
 
-	void RemoveObjectPropertyFromParameterization(UObject* Object, const TArray<struct FDataprepPropertyLink>& InPropertyChain);
+	void RemoveObjectPropertyFromParameterization(UDataprepParameterizableObject* Object, const TArray<struct FDataprepPropertyLink>& InPropertyChain);
 
 	UDataprepParameterization* GetDataprepParameterization() { return Parameterization; }
 
diff --git a/Engine/Plugins/Experimental/DataPrepEditor/Source/DataPrepCore/Public/DataPrepOperation.h b/Engine/Plugins/Experimental/DataPrepEditor/Source/DataPrepCore/Public/DataPrepOperation.h
--- a/Engine/Plugins/Experimental/DataPrepEditor/Source/DataPrepCore/Public/DataPrepOperation.h
+++ b/Engine/Plugins/Experimental/DataPrepEditor/Source/DataPrepCore/Public/DataPrepOperation.h
@@ -2,8 +2,9 @@
 
 #pragma once
 
-#include "CoreMinimal.h"
+#include "DataprepParameterizableObject.h"
 
+#include "CoreMinimal.h"
 #include "Delegates/DelegateCombinations.h"
 #include "UObject/Object.h"
 #include "UObject/ObjectMacros.h"
@@ -105,7 +106,7 @@ struct FDataprepOperationContext
 
 // Todo add class wide comment
 UCLASS(Experimental, Abstract, Blueprintable)
-class DATAPREPCORE_API UDataprepOperation : public UObject
+class DATAPREPCORE_API UDataprepOperation : public UDataprepParameterizableObject
 {
 	GENERATED_BODY()
 
diff --git a/Engine/Plugins/Experimental/DataPrepEditor/Source/DataPrepCore/Public/DataprepParameterizableObject.h b/Engine/Plugins/Experimental/DataPrepEditor/Source/DataPrepCore/Public/DataprepParameterizableObject.h
new file mode 100644
--- /dev/null
+++ b/Engine/Plugins/Experimental/DataPrepEditor/Source/DataPrepCore/Public/DataprepParameterizableObject.h
@@ -0,0 +1,23 @@
+// Copyright 1998-2019 Epic Games, Inc. All Rights Reserved.
+
+#pragma once
+
+#include "CoreMinimal.h"
+#include "UObject/Object.h"
+
+#include "DataprepParameterizableObject.generated.h"
+
+struct FPropertyChangedChainEvent;
+
+/**
+ * The base class of all the object that can interact with the dataprep parameterization
+ * This include all the parameterizable object and the parameterization object itself
+ */
+UCLASS()
+class DATAPREPCORE_API UDataprepParameterizableObject : public UObject
+{
+public:
+	GENERATED_BODY()
+
+	virtual void PostEditChangeChainProperty(FPropertyChangedChainEvent& PropertyChangedEvent) override;
+};
diff --git a/Engine/Plugins/Experimental/DataPrepEditor/Source/DataPrepCore/Public/Parameterization/DataprepParameterizationUtils.h b/Engine/Plugins/Experimental/DataPrepEditor/Source/DataPrepCore/Public/Parameterization/DataprepParameterizationUtils.h
--- a/Engine/Plugins/Experimental/DataPrepEditor/Source/DataPrepCore/Public/Parameterization/DataprepParameterizationUtils.h
+++ b/Engine/Plugins/Experimental/DataPrepEditor/Source/DataPrepCore/Public/Parameterization/DataprepParameterizationUtils.h
@@ -35,7 +35,7 @@ struct FDataprepPropertyLink
 {
 	GENERATED_BODY()
 
-	FDataprepPropertyLink( UProperty* InCachedProperty, FName InPropertyName, int32 InContenerIndex )
+	FDataprepPropertyLink( UProperty* InCachedProperty, const FName& InPropertyName, int32 InContenerIndex )
 		: CachedProperty( MakeWeakObjectPtr(InCachedProperty) )
 		, PropertyName( InPropertyName )
 		, ContainerIndex( InContenerIndex )
@@ -69,6 +69,12 @@ public:
 	 */
 	static TArray<FDataprepPropertyLink> MakePropertyChain(TSharedPtr<IPropertyHandle> PropertyHandle);
 
+	/**
+	 * Take a Property Changed Event and extract a dataprep property link chain from it
+	 * @param PropertyChangedEvent The event
+	 * @return A non empty array if we were able make a compatible property chain
+	 */
+	static TArray<FDataprepPropertyLink>  MakePropertyChain(FPropertyChangedChainEvent& PropertyChangedEvent);
 
 	/**
 	 * Take a already existing parameterization context and create a new version including the handle.
diff --git a/Engine/Plugins/Experimental/DataPrepEditor/Source/DataPrepCore/Public/SelectionSystem/DataprepFetcher.h b/Engine/Plugins/Experimental/DataPrepEditor/Source/DataPrepCore/Public/SelectionSystem/DataprepFetcher.h
--- a/Engine/Plugins/Experimental/DataPrepEditor/Source/DataPrepCore/Public/SelectionSystem/DataprepFetcher.h
+++ b/Engine/Plugins/Experimental/DataPrepEditor/Source/DataPrepCore/Public/SelectionSystem/DataprepFetcher.h
@@ -2,6 +2,8 @@
 
 #pragma once
 
+#include "DataprepParameterizableObject.h"
+
 #include "CoreMinimal.h"
 #include "UObject/Object.h"
 #include "UObject/ObjectMacros.h"
@@ -11,10 +13,10 @@
 /**
  * The Dataprep fetcher is a base class for Dataprep Selection system.
  * The responsibility of fetcher is return a certain type of data for a object. (Look at DataprepStringFetcher for a example)
- * This absract base class exist currently for the discovery process and some compile time validation
+ * This abstract base class exist currently for the discovery process, some compile time validation and to propagate UDataprepParameterizableObject
  */
 UCLASS(Abstract)
-class DATAPREPCORE_API UDataprepFetcher : public UObject
+class DATAPREPCORE_API UDataprepFetcher : public UDataprepParameterizableObject
 {
 	GENERATED_BODY()
 
diff --git a/Engine/Plugins/Experimental/DataPrepEditor/Source/DataPrepCore/Public/SelectionSystem/DataprepFilter.h b/Engine/Plugins/Experimental/DataPrepEditor/Source/DataPrepCore/Public/SelectionSystem/DataprepFilter.h
--- a/Engine/Plugins/Experimental/DataPrepEditor/Source/DataPrepCore/Public/SelectionSystem/DataprepFilter.h
+++ b/Engine/Plugins/Experimental/DataPrepEditor/Source/DataPrepCore/Public/SelectionSystem/DataprepFilter.h
@@ -3,6 +3,7 @@
 #pragma once
 
 #include "DataprepFetcher.h"
+#include "DataprepParameterizableObject.h"
 
 #include "CoreMinimal.h"
 #include "Templates/SubclassOf.h"
@@ -16,7 +17,7 @@
  * It's main responsibility is to filter a array of object and to return the selected objects
  */
 UCLASS(Abstract)
-class DATAPREPCORE_API UDataprepFilter : public UObject
+class DATAPREPCORE_API UDataprepFilter : public UDataprepParameterizableObject
 {
 	GENERATED_BODY()
 
diff --git a/Engine/Plugins/Experimental/DataPrepEditor/Source/DataPrepEditor/Private/Widgets/DataprepWidgets.cpp b/Engine/Plugins/Experimental/DataPrepEditor/Source/DataPrepEditor/Private/Widgets/DataprepWidgets.cpp
--- a/Engine/Plugins/Experimental/DataPrepEditor/Source/DataPrepEditor/Private/Widgets/DataprepWidgets.cpp
+++ b/Engine/Plugins/Experimental/DataPrepEditor/Source/DataPrepEditor/Private/Widgets/DataprepWidgets.cpp
@@ -7,6 +7,7 @@
 #include "DataprepAssetView.h"
 #include "DataprepEditorLogCategory.h"
 #include "DataprepEditorUtils.h"
+#include "DataprepParameterizableObject.h"
 #include "Parameterization/DataprepParameterizationUtils.h"
 
 #include "ContentBrowserModule.h"
@@ -438,7 +439,7 @@ TSharedRef< SWidget > SDataprepDetailsView::CreateDefaultWidget( TSharedPtr< SWi
 							if ( UDataprepAsset* DataprepAsset = DataprepAssetForParameterization.Get() )
 							{
 								FScopedTransaction Transaction( LOCTEXT("DataprepBindingToParameterization","Adding Parameter") );
-								DataprepAssetForParameterization->BindObjectPropertyToParameterization( DetailedObject, PropertyChain, *PropertyChain.Last().CachedProperty->GetDisplayNameText().ToString() );
+								DataprepAssetForParameterization->BindObjectPropertyToParameterization( DetailedObjectAsParameterizable, PropertyChain, *PropertyChain.Last().CachedProperty->GetDisplayNameText().ToString() );
 								bRefreshObjectToDisplay = true;
 							}
 
@@ -464,7 +465,7 @@ TSharedRef< SWidget > SDataprepDetailsView::CreateDefaultWidget( TSharedPtr< SWi
 						if (UDataprepAsset * DataprepAsset = DataprepAssetForParameterization.Get())
 						{
 							FScopedTransaction Transaction( LOCTEXT("DataprepRemoveBindingFromParameterization","Removing Parameter") );
-							DataprepAssetForParameterization->RemoveObjectPropertyFromParameterization( DetailedObject, PropertyChain );
+							DataprepAssetForParameterization->RemoveObjectPropertyFromParameterization( DetailedObjectAsParameterizable, PropertyChain );
 							bRefreshObjectToDisplay = true;
 						}
 
@@ -538,6 +539,11 @@ void SDataprepDetailsView::OnObjectReplaced(const TMap<UObject*, UObject*>& Repl
 	if ( UObject * const* ObjectPtr = ReplacementObjectMap.Find( DetailedObject ) )
 	{
 		DetailedObject = *ObjectPtr;
+		if ( DetailedObject->IsA<UDataprepParameterizableObject>() )
+		{
+			DetailedObjectAsParameterizable = static_cast<UDataprepParameterizableObject*>( DetailedObject );
+		}
+
 		bRefreshObjectToDisplay = true;
 	}
 }
@@ -600,7 +606,7 @@ void SDataprepDetailsView::AddWidgets( const TArray< TSharedRef< IDetailTreeNode
 		{
 			if ( UDataprepAsset* DataprepAsset = DataprepAssetForParameterization.Get() )
 			{
-				if ( DataprepAsset->IsObjectPropertyBinded( DetailedObject, CurrentParameterizationContext.PropertyChain ) )
+				if ( DataprepAsset->IsObjectPropertyBinded( DetailedObjectAsParameterizable, CurrentParameterizationContext.PropertyChain ) )
 				{
 					CurrentParameterizationContext.State = EParametrizationState::IsParameterized;
 				}
@@ -718,6 +724,10 @@ void SDataprepDetailsView::Construct(const FArguments& InArgs)
 {
 	bRefreshObjectToDisplay = false;
 	DetailedObject = InArgs._Object;
+	if ( DetailedObject->IsA<UDataprepParameterizableObject>() )
+	{
+		DetailedObjectAsParameterizable = static_cast<UDataprepParameterizableObject*>( DetailedObject );
+	}
 
 	if (InArgs._ColumnSizeData.IsValid())
 	{
@@ -767,7 +777,7 @@ void SDataprepDetailsView::Construct()
 		}
 
 		FDataprepParameterizationContext ParameterizationContext;
-		ParameterizationContext.State = DataprepAsset ? EParametrizationState::CanBeParameterized : EParametrizationState::InvalidForParameterization;
+		ParameterizationContext.State = DataprepAsset && DetailedObjectAsParameterizable ? EParametrizationState::CanBeParameterized : EParametrizationState::InvalidForParameterization;
 		DataprepAssetForParameterization = DataprepAsset;
 
 		TSharedPtr<SGridPanel> GridPanel = SNew(SGridPanel).FillColumn( 0.0f, 1.0f );
@@ -849,8 +859,19 @@ void SDataprepDetailsView::SetObjectToDisplay(UObject& Object)
 	if ( DetailedObject != NewObjectToDisplay )
 	{
 		DetailedObject = NewObjectToDisplay;
+		if ( DetailedObject->IsA<UDataprepParameterizableObject>() )
+		{
+			DetailedObjectAsParameterizable = static_cast<UDataprepParameterizableObject*>( DetailedObject );
+		}
 		bRefreshObjectToDisplay = true;
 	}
 }
 
+void SDataprepDetailsView::AddReferencedObjects(FReferenceCollector& Collector)
+{
+	Collector.AddReferencedObject( DetailedObject );
+	Collector.AddReferencedObject( DetailedObjectAsParameterizable );
+	Collector.AddReferencedObjects( TrackedProperties );
+}
+
 #undef LOCTEXT_NAMESPACE
diff --git a/Engine/Plugins/Experimental/DataPrepEditor/Source/DataPrepEditor/Private/Widgets/DataprepWidgets.h b/Engine/Plugins/Experimental/DataPrepEditor/Source/DataPrepEditor/Private/Widgets/DataprepWidgets.h
--- a/Engine/Plugins/Experimental/DataPrepEditor/Source/DataPrepEditor/Private/Widgets/DataprepWidgets.h
+++ b/Engine/Plugins/Experimental/DataPrepEditor/Source/DataPrepEditor/Private/Widgets/DataprepWidgets.h
@@ -11,6 +11,7 @@
 class SEditableTextBox;
 class UDataprepAsset;
 class UDataprepContentConsumer;
+class UDataprepParameterizableObject;
 
 struct FDataprepParameterizationContext;
 struct FDataprepPropertyLink;
@@ -82,7 +83,7 @@ private:
 	float ColumnWidth;
 };
 
-class SDataprepDetailsView : public SCompoundWidget
+class SDataprepDetailsView : public SCompoundWidget, public FGCObject
 {
 	using Super = SCompoundWidget;
 
@@ -103,6 +104,10 @@ public:
 
 	void SetObjectToDisplay(UObject& Object);
 
+	// FGCObject interface
+	virtual void AddReferencedObjects(FReferenceCollector& Collector) override;
+	// End of FGCObject interface
+
 protected:
 	// Begin SWidget overrides.
 	// #ueent_todo: This is temporary until we find a better solution to the splitter issue
@@ -158,6 +163,9 @@ private:
 	/** Object to be detailed */
 	UObject* DetailedObject;
 
+	/** Not null if the detailed object is parameterizable */
+	UDataprepParameterizableObject* DetailedObjectAsParameterizable;
+
 	/** Array properties tracked for changes */
 	TSet< UProperty* > TrackedProperties;
 	
