Merging CL #9614506 from //UE4/Dev-Enterprise

Dataprep: Rework filter naming

diff --git a/Engine/Plugins/Enterprise/DatasmithImporter/Source/DatasmithImporter/Private/DatasmithDataprepFetcher.cpp b/Engine/Plugins/Enterprise/DatasmithImporter/Source/DatasmithImporter/Private/DatasmithDataprepFetcher.cpp
--- a/Engine/Plugins/Enterprise/DatasmithImporter/Source/DatasmithImporter/Private/DatasmithDataprepFetcher.cpp
+++ b/Engine/Plugins/Enterprise/DatasmithImporter/Source/DatasmithImporter/Private/DatasmithDataprepFetcher.cpp
@@ -4,6 +4,8 @@
 
 #include "DatasmithContentBlueprintLibrary.h"
 
+#define LOCTEXT_NAMESPACE "DatasmithDataprepFetcher"
+
 /* UDataprepStringMetadataValueFetcher methods
  *****************************************************************************/
 FString UDatasmithStringMetadataValueFetcher::Fetch_Implementation(const UObject* Object, bool& bOutFetchSucceded) const
@@ -18,7 +20,14 @@ FString UDatasmithStringMetadataValueFetcher::Fetch_Implementation(const UObject
 	return {};
 }
 
+FText UDatasmithStringMetadataValueFetcher::GetNodeDisplayFetcherName_Implementation() const
+{
+	return LOCTEXT("NodeDisplayFetcher_MetadataValue", "Metadata");
+}
+
 bool UDatasmithStringMetadataValueFetcher::IsThreadSafe() const
 {
 	return true;
-}
\ No newline at end of file
+}
+
+#undef LOCTEXT_NAMESPACE
diff --git a/Engine/Plugins/Enterprise/DatasmithImporter/Source/DatasmithImporter/Private/DatasmithDataprepFetcher.h b/Engine/Plugins/Enterprise/DatasmithImporter/Source/DatasmithImporter/Private/DatasmithDataprepFetcher.h
--- a/Engine/Plugins/Enterprise/DatasmithImporter/Source/DatasmithImporter/Private/DatasmithDataprepFetcher.h
+++ b/Engine/Plugins/Enterprise/DatasmithImporter/Source/DatasmithImporter/Private/DatasmithDataprepFetcher.h
@@ -18,6 +18,7 @@ public:
 	//~ End of UDataprepStringFetcher interface
 
 	//~ UDataprepFetcher interface
+	virtual FText GetNodeDisplayFetcherName_Implementation() const;
 	virtual bool IsThreadSafe() const final;
 	//~ End of UDataprepFetcher interface
 
diff --git a/Engine/Plugins/Experimental/DataPrepEditor/Source/DataPrepCore/Private/SelectionSystem/DataprepFetcher.cpp b/Engine/Plugins/Experimental/DataPrepEditor/Source/DataPrepCore/Private/SelectionSystem/DataprepFetcher.cpp
--- a/Engine/Plugins/Experimental/DataPrepEditor/Source/DataPrepCore/Private/SelectionSystem/DataprepFetcher.cpp
+++ b/Engine/Plugins/Experimental/DataPrepEditor/Source/DataPrepCore/Private/SelectionSystem/DataprepFetcher.cpp
@@ -7,6 +7,11 @@ FText UDataprepFetcher::GetDisplayFetcherName_Implementation() const
 	return GetClass()->GetDisplayNameText();
 }
 
+FText UDataprepFetcher::GetNodeDisplayFetcherName_Implementation() const
+{
+	return GetClass()->GetDisplayNameText();
+}
+
 FText UDataprepFetcher::GetTooltipText_Implementation() const
 {
 	return GetClass()->GetToolTipText();
diff --git a/Engine/Plugins/Experimental/DataPrepEditor/Source/DataPrepCore/Private/SelectionSystem/DataprepStringsArrayFilter.cpp b/Engine/Plugins/Experimental/DataPrepEditor/Source/DataPrepCore/Private/SelectionSystem/DataprepStringsArrayFilter.cpp
new file mode 100644
--- /dev/null
+++ b/Engine/Plugins/Experimental/DataPrepEditor/Source/DataPrepCore/Private/SelectionSystem/DataprepStringsArrayFilter.cpp
@@ -0,0 +1,110 @@
+// Copyright 1998-2019 Epic Games, Inc. All Rights Reserved.
+
+#include "SelectionSystem/DataprepStringsArrayFilter.h"
+
+#include "DataprepCoreLogCategory.h"
+#include "SelectionSystem/DataprepSelectionSystemUtils.h"
+#include "SelectionSystem/DataprepStringsArrayFetcher.h"
+
+bool UDataprepStringsArrayFilter::Filter(const TArray<FString>& StringArray) const
+{
+	bool bResult = false;
+
+	for (int Index = 0; Index < StringArray.Num(); ++Index)
+	{
+		switch (StringMatchingCriteria)
+		{
+			case EDataprepStringMatchType::Contains:
+				bResult = StringArray[Index].Equals(UserString);
+				break;
+			case EDataprepStringMatchType::ExactMatch:
+				bResult = StringArray[Index].Contains(UserString);
+				break;
+			case EDataprepStringMatchType::MatchesWildcard:
+				bResult = StringArray[Index].MatchesWildcard(UserString);
+				break;
+		}
+
+		if (bResult)
+		{
+			break;
+		}
+	}
+
+	return bResult;
+}
+
+TArray<UObject*> UDataprepStringsArrayFilter::FilterObjects(const TArray<UObject *>& Objects) const
+{
+	if ( StringsArrayFetcher )
+	{
+		return DataprepSelectionSystemUtils::FilterObjects< UDataprepStringsArrayFilter, UDataprepStringsArrayFetcher, TArray<FString> >( *this, *StringsArrayFetcher, Objects );
+	}
+
+	ensure(false);
+	UE_LOG(LogDataprepCore, Error, TEXT("UDataprepStringsArrayFilter::FilterObjects: There was no Fetcher"));
+	return {};
+}
+
+TSubclassOf<UDataprepFetcher> UDataprepStringsArrayFilter::GetAcceptedFetcherClass() const
+{
+	return UDataprepStringsArrayFetcher::StaticClass();
+}
+
+void UDataprepStringsArrayFilter::SetFetcher(const TSubclassOf<UDataprepFetcher>& FetcherClass)
+{
+	UClass* NewFetcherClass = FetcherClass;
+
+	if ( NewFetcherClass && NewFetcherClass->IsChildOf( GetAcceptedFetcherClass() ) )
+	{
+		UClass* OldFetcherClass = StringsArrayFetcher ? StringsArrayFetcher->GetClass() : nullptr;
+		if ( NewFetcherClass != OldFetcherClass )
+		{
+			Modify();
+			StringsArrayFetcher = NewObject< UDataprepStringsArrayFetcher >( this, NewFetcherClass, NAME_None, RF_Transactional );
+		}
+	}
+	else
+	{
+		ensure( false );
+		UE_LOG( LogDataprepCore, Error, TEXT("UDataprepStringFilter::SetFetcher: The Fetcher Class is not compatible") );
+	}
+}
+
+UDataprepFetcher* UDataprepStringsArrayFilter::GetFetcher() const
+{
+	return StringsArrayFetcher;
+}
+
+FText UDataprepStringsArrayFilter::GetFilterCategoryText() const
+{
+	return NSLOCTEXT("DataprepSringsArrayFilter", "StringFilterCategory", "String");
+}
+
+EDataprepStringMatchType UDataprepStringsArrayFilter::GetStringMatchingCriteria() const
+{
+	return StringMatchingCriteria;
+}
+
+FString UDataprepStringsArrayFilter::GetUserString() const
+{
+	return UserString;
+}
+
+void UDataprepStringsArrayFilter::SetStringMatchingCriteria(EDataprepStringMatchType InStringMatchingCriteria)
+{
+	if (StringMatchingCriteria != InStringMatchingCriteria)
+	{
+		Modify();
+		StringMatchingCriteria = InStringMatchingCriteria;
+	}
+}
+
+void UDataprepStringsArrayFilter::SetUserString(FString InUserString)
+{
+	if (UserString != InUserString)
+	{
+		Modify();
+		UserString = InUserString;
+	}
+}
diff --git a/Engine/Plugins/Experimental/DataPrepEditor/Source/DataPrepCore/Public/SelectionSystem/DataprepFetcher.h b/Engine/Plugins/Experimental/DataPrepEditor/Source/DataPrepCore/Public/SelectionSystem/DataprepFetcher.h
--- a/Engine/Plugins/Experimental/DataPrepEditor/Source/DataPrepCore/Public/SelectionSystem/DataprepFetcher.h
+++ b/Engine/Plugins/Experimental/DataPrepEditor/Source/DataPrepCore/Public/SelectionSystem/DataprepFetcher.h
@@ -32,6 +32,12 @@ public:
 	UFUNCTION(BlueprintNativeEvent,  Category = "Display")
 	FText GetDisplayFetcherName() const;
 
+	/**
+	 * The name displayed on node title.
+	 */
+	UFUNCTION(BlueprintNativeEvent, Category = "Display")
+	FText GetNodeDisplayFetcherName() const;
+
 	/**
 	 * Allows to change the tooltip of the fetcher for the ui if needed.
 	 */
@@ -47,6 +53,7 @@ public:
 
 	// The Native way to override the blueprint native events above
 	virtual FText GetDisplayFetcherName_Implementation() const;
+	virtual FText GetNodeDisplayFetcherName_Implementation() const;
 	virtual FText GetTooltipText_Implementation() const;
 	virtual FText GetAdditionalKeyword_Implementation() const;
 };
diff --git a/Engine/Plugins/Experimental/DataPrepEditor/Source/DataPrepCore/Public/SelectionSystem/DataprepStringsArrayFetcher.h b/Engine/Plugins/Experimental/DataPrepEditor/Source/DataPrepCore/Public/SelectionSystem/DataprepStringsArrayFetcher.h
new file mode 100644
--- /dev/null
+++ b/Engine/Plugins/Experimental/DataPrepEditor/Source/DataPrepCore/Public/SelectionSystem/DataprepStringsArrayFetcher.h
@@ -0,0 +1,42 @@
+// Copyright 1998-2019 Epic Games, Inc. All Rights Reserved.
+
+#pragma once
+
+#include "SelectionSystem/DataprepFetcher.h"
+
+#include "CoreMinimal.h"
+#include "UObject/Object.h"
+#include "UObject/ObjectMacros.h"
+
+#include "DataprepStringsArrayFetcher.generated.h"
+
+UCLASS(Abstract, Blueprintable)
+class DATAPREPCORE_API UDataprepStringsArrayFetcher : public UDataprepFetcher
+{
+	GENERATED_BODY()
+
+public:
+
+	/**
+	 * This function is called when the fetcher is executed.
+	 * If you are defining your fetcher in Blueprint this is the function to override.
+	 * @param Object The object from which the fetcher should try to retrieve the string
+	 * @param bOutFetchSucceded If the fetcher managed to retrieve the string from the object this bool must be set to true
+	 * @return The fetched string
+	 */
+	UFUNCTION(BlueprintNativeEvent)
+	TArray<FString> Fetch(const UObject* Object, bool& bOutFetchSucceded) const;
+
+	/**
+	 * This function is the same has Fetch, but it's the extension point for an operation defined in c++.
+	 * It will be called on the fetcher execution.
+	 * @param Object The object from which the fetcher should try to retrieve the string
+	 * @param bOutFetchSucceded If the fetcher managed to retrieve the string from the object this bool must be set to true
+	 * @return The fetched string
+	 */
+	virtual TArray<FString> Fetch_Implementation(const UObject* Object, bool& bOutFetchSucceded) const
+	{
+		bOutFetchSucceded = false;
+		return {};
+	} 
+};
diff --git a/Engine/Plugins/Experimental/DataPrepEditor/Source/DataPrepCore/Public/SelectionSystem/DataprepStringsArrayFilter.h b/Engine/Plugins/Experimental/DataPrepEditor/Source/DataPrepCore/Public/SelectionSystem/DataprepStringsArrayFilter.h
new file mode 100644
--- /dev/null
+++ b/Engine/Plugins/Experimental/DataPrepEditor/Source/DataPrepCore/Public/SelectionSystem/DataprepStringsArrayFilter.h
@@ -0,0 +1,46 @@
+// Copyright 1998-2019 Epic Games, Inc. All Rights Reserved.
+
+#pragma once
+
+#include "SelectionSystem/DataprepStringFilter.h"
+
+#include "DataprepStringsArrayFilter.generated.h"
+
+class UDataprepStringsArrayFetcher;
+
+UCLASS()
+class DATAPREPCORE_API UDataprepStringsArrayFilter : public UDataprepFilter
+{
+	GENERATED_BODY()
+
+public:
+	bool Filter(const TArray<FString>& StringArray) const;
+
+	//~ Begin UDataprepFilter Interface
+	virtual TArray<UObject*> FilterObjects(const TArray<UObject*>& Objects) const override;
+	virtual bool IsThreadSafe() const override { return true; }
+	virtual FText GetFilterCategoryText() const override;
+	virtual TSubclassOf<UDataprepFetcher> GetAcceptedFetcherClass() const override;
+	virtual void SetFetcher(const TSubclassOf<UDataprepFetcher>& FetcherClass) override;
+	virtual UDataprepFetcher* GetFetcher() const override;
+	//~ Begin UDataprepFilter Interface
+
+	EDataprepStringMatchType GetStringMatchingCriteria() const;
+	FString GetUserString() const;
+
+	void SetStringMatchingCriteria(EDataprepStringMatchType StringMatchingCriteria);
+	void SetUserString(FString UserString);
+
+private:
+	// The criteria selected by the user
+	UPROPERTY()
+	EDataprepStringMatchType StringMatchingCriteria;
+
+	// The string entered by the user
+	UPROPERTY()
+	FString UserString;
+
+	// The source of string selected by the user
+	UPROPERTY()
+	UDataprepStringsArrayFetcher* StringsArrayFetcher;
+};
diff --git a/Engine/Plugins/Experimental/DataPrepEditor/Source/DataPrepEditor/Private/Widgets/Action/SDataprepBoolFilter.cpp b/Engine/Plugins/Experimental/DataPrepEditor/Source/DataPrepEditor/Private/Widgets/Action/SDataprepBoolFilter.cpp
--- a/Engine/Plugins/Experimental/DataPrepEditor/Source/DataPrepEditor/Private/Widgets/Action/SDataprepBoolFilter.cpp
+++ b/Engine/Plugins/Experimental/DataPrepEditor/Source/DataPrepEditor/Private/Widgets/Action/SDataprepBoolFilter.cpp
@@ -9,13 +9,4 @@
 
 void SDataprepBoolFilter::Construct(const FArguments& InArgs, UDataprepBoolFilter& InFilter)
 {
-	ChildSlot
-	[
-		SNew(SBox)
-		.MinDesiredWidth( 400.f )
-		.Padding( 5.f )
-		[
-			SNew( SDataprepFetcherSelector, InFilter )
-		]
-	];
 }
diff --git a/Engine/Plugins/Experimental/DataPrepEditor/Source/DataPrepEditor/Private/Widgets/Action/SDataprepFilter.cpp b/Engine/Plugins/Experimental/DataPrepEditor/Source/DataPrepEditor/Private/Widgets/Action/SDataprepFilter.cpp
--- a/Engine/Plugins/Experimental/DataPrepEditor/Source/DataPrepEditor/Private/Widgets/Action/SDataprepFilter.cpp
+++ b/Engine/Plugins/Experimental/DataPrepEditor/Source/DataPrepEditor/Private/Widgets/Action/SDataprepFilter.cpp
@@ -8,6 +8,7 @@
 #include "SelectionSystem/DataprepFilter.h"
 #include "SelectionSystem/DataprepFloatFilter.h"
 #include "SelectionSystem/DataprepStringFilter.h"
+#include "SelectionSystem/DataprepStringsArrayFilter.h"
 #include "Widgets/Action/SDataprepBoolFilter.h"
 #include "Widgets/Action/SDataprepFloatFilter.h"
 #include "Widgets/Action/SDataprepStringFilter.h"
@@ -46,44 +47,16 @@ FText SDataprepFilter::GetBlockTitle() const
 {
 	if ( Filter )
 	{
-		UClass* Class = Filter->GetClass();
-		if ( Class == UDataprepStringFilter::StaticClass() )
-		{
-			if ( Filter->IsExcludingResult() )
-			{
-				static FText StringExcludingFilterTitle =  LOCTEXT("StringExcludingFilterTitle", "Exclude by String");
-				return StringExcludingFilterTitle;
-			}
-			else
-			{
-				static FText StringFilterTitle = LOCTEXT("StringFilterTitle", "Select by String");
-				return StringFilterTitle;
-			}
-		}
-		else if ( Class == UDataprepBoolFilter::StaticClass() )
-		{
-			if ( Filter->IsExcludingResult() )
-			{
-				static FText BoolExcludingFilterTitle = LOCTEXT("BoolExcludingFilterTitle", "Exclude by Condition");
-				return BoolExcludingFilterTitle;
-			}
-			else
-			{
-				static FText BoolFilterTitle = LOCTEXT("BoolFilterTitle", "Select by Condition");
-				return BoolFilterTitle;
-			}
-		}
-		else if ( Class == UDataprepFloatFilter::StaticClass() )
+		UDataprepFetcher* Fetcher = Filter->GetFetcher();
+		if ( Fetcher )
 		{
 			if ( Filter->IsExcludingResult() )
 			{
-				static FText FloatExcludingFilterTitle = LOCTEXT("FloatExcludingFilterTitle", "Exclude by Float");
-				return FloatExcludingFilterTitle;
+				return FText::Format( LOCTEXT("ExcludingFilterTitle", "Exclude by {0}"), { Fetcher->GetNodeDisplayFetcherName() } );
 			}
 			else
 			{
-				static FText FloatFilterTitle = LOCTEXT("FloatFilterTitle", "Select by Float");
-				return FloatFilterTitle;
+				return FText::Format( LOCTEXT("SelectingFilterTitle", "Filter by {0}"), { Fetcher->GetNodeDisplayFetcherName() });
 			}
 		}
 	}
@@ -106,7 +79,7 @@ TSharedRef<SWidget> SDataprepFilter::GetTitleWidget()
 
 TSharedRef<SWidget> SDataprepFilter::GetContentWidget()
 {
-	TSharedPtr< SWidget > FilterWidget;
+	TSharedPtr< SWidget > FilterWidget = SNullWidget::NullWidget;
 
 	if ( Filter )
 	{
@@ -114,7 +87,11 @@ TSharedRef<SWidget> SDataprepFilter::GetContentWidget()
 		// This down casting implementation is faster then using Cast<UDataprepStringFilter>( Filter ) 
 		if ( Class ==  UDataprepStringFilter::StaticClass() )
 		{
-			SAssignNew( FilterWidget, SDataprepStringFilter, *static_cast< UDataprepStringFilter* >( Filter ) );
+			SAssignNew( FilterWidget, SDataprepStringFilter< UDataprepStringFilter >, *static_cast< UDataprepStringFilter* >( Filter ) );
+		}
+		else if (Class == UDataprepStringsArrayFilter::StaticClass())
+		{
+			SAssignNew(FilterWidget, SDataprepStringFilter< UDataprepStringsArrayFilter >, *static_cast<UDataprepStringsArrayFilter*>(Filter));
 		}
 		else if ( Class == UDataprepBoolFilter::StaticClass() )
 		{
@@ -125,10 +102,6 @@ TSharedRef<SWidget> SDataprepFilter::GetContentWidget()
 			SAssignNew( FilterWidget, SDataprepFloatFilter, *static_cast< UDataprepFloatFilter* >( Filter ) );
 		}
 	}
-	else
-	{
-		FilterWidget = SNullWidget::NullWidget;
-	}
 
 	return SNew( SVerticalBox )
 		+ SVerticalBox::Slot()
diff --git a/Engine/Plugins/Experimental/DataPrepEditor/Source/DataPrepEditor/Private/Widgets/Action/SDataprepFloatFilter.cpp b/Engine/Plugins/Experimental/DataPrepEditor/Source/DataPrepEditor/Private/Widgets/Action/SDataprepFloatFilter.cpp
--- a/Engine/Plugins/Experimental/DataPrepEditor/Source/DataPrepEditor/Private/Widgets/Action/SDataprepFloatFilter.cpp
+++ b/Engine/Plugins/Experimental/DataPrepEditor/Source/DataPrepEditor/Private/Widgets/Action/SDataprepFloatFilter.cpp
@@ -39,11 +39,6 @@ void SDataprepFloatFilter::Construct(const FArguments& InArgs, UDataprepFloatFil
 				SNew( SHorizontalBox )
 				+ SHorizontalBox::Slot()
 				.Padding( 5.f )
-				[
-					SNew( SDataprepFetcherSelector, InFilter )
-				]
-				+ SHorizontalBox::Slot()
-				.Padding( 5.f )
 				[
 					SAssignNew( FloatMatchingCriteriaWidget, SComboBox< TSharedPtr< FListEntry > > )
 					.OptionsSource( &FloatMatchingOptions )
diff --git a/Engine/Plugins/Experimental/DataPrepEditor/Source/DataPrepEditor/Private/Widgets/Action/SDataprepStringFilter.cpp b/Engine/Plugins/Experimental/DataPrepEditor/Source/DataPrepEditor/Private/Widgets/Action/SDataprepStringFilter.cpp
--- a/Engine/Plugins/Experimental/DataPrepEditor/Source/DataPrepEditor/Private/Widgets/Action/SDataprepStringFilter.cpp
+++ b/Engine/Plugins/Experimental/DataPrepEditor/Source/DataPrepEditor/Private/Widgets/Action/SDataprepStringFilter.cpp
@@ -6,6 +6,7 @@
 #include "SchemaActions/DataprepFetcherMenuActionCollector.h"
 #include "SchemaActions/DataprepSchemaAction.h"
 #include "SelectionSystem/DataprepStringFilter.h"
+#include "SelectionSystem/DataprepStringsArrayFilter.h"
 #include "Widgets/Action/DataprepActionWidgetsUtils.h"
 #include "Widgets/Action/SDataprepFetcherSelector.h"
 
@@ -19,7 +20,8 @@
 
 #define LOCTEXT_NAMESPACE "DataprepStringFilter"
 
-void SDataprepStringFilter::Construct(const FArguments& InArgs, UDataprepStringFilter& InFilter)
+template <class FilterType>
+void SDataprepStringFilter<FilterType>::Construct(const FArguments& InArgs, FilterType& InFilter)
 {
 	Filter = &InFilter;
 	OldUserString = Filter->GetUserString();
@@ -34,11 +36,6 @@ void SDataprepStringFilter::Construct(const FArguments& InArgs, UDataprepStringF
 			SNew( SHorizontalBox )
 			+ SHorizontalBox::Slot()
 			.Padding( 5.f )
-			[
-				SNew( SDataprepFetcherSelector, InFilter )
-			]
-			+ SHorizontalBox::Slot()
-			.Padding( 5.f )
 			[
 				SAssignNew( StringMatchingCriteriaWidget, SComboBox< TSharedPtr< FListEntry > > )
 				.OptionsSource( &StringMatchingOptions )
@@ -66,28 +63,32 @@ void SDataprepStringFilter::Construct(const FArguments& InArgs, UDataprepStringF
 	];
 }
 
-TSharedRef<SWidget> SDataprepStringFilter::OnGenerateWidgetForMatchingCriteria(TSharedPtr<FListEntry> ListEntry) const
+template <class FilterType>
+TSharedRef<SWidget> SDataprepStringFilter<FilterType>::OnGenerateWidgetForMatchingCriteria(TSharedPtr<FListEntry> ListEntry) const
 {
 	return SNew(STextBlock)
 		.Text( ListEntry->Get<0>() )
 		.ToolTipText( ListEntry->Get<1>() );
 }
 
-FText SDataprepStringFilter::GetSelectedCriteriaText() const
+template <class FilterType>
+FText SDataprepStringFilter<FilterType>::GetSelectedCriteriaText() const
 {
 	check( Filter );
 	UEnum* Enum = StaticEnum< EDataprepStringMatchType >();
 	return Enum->GetDisplayNameTextByValue( static_cast<uint8>( Filter->GetStringMatchingCriteria() ) );
 }
 
-FText SDataprepStringFilter::GetSelectedCriteriaTooltipText() const
+template <class FilterType>
+FText SDataprepStringFilter<FilterType>::GetSelectedCriteriaTooltipText() const
 {
 	check( Filter );
 	UEnum* Enum = StaticEnum< EDataprepStringMatchType >();
 	return  Enum->GetToolTipTextByIndex( Enum->GetIndexByValue( static_cast<uint8>( Filter->GetStringMatchingCriteria() ) ) );
 }
 
-void SDataprepStringFilter::OnCriteriaComboBoxOpenning()
+template <class FilterType>
+void SDataprepStringFilter<FilterType>::OnCriteriaComboBoxOpenning()
 {
 	check( Filter );
 	UEnum* Enum = StaticEnum< EDataprepStringMatchType >();
@@ -107,7 +108,8 @@ void SDataprepStringFilter::OnCriteriaComboBoxOpenning()
 	StringMatchingCriteriaWidget->SetSelectedItem( ItemToSelect );
 }
 
-void SDataprepStringFilter::OnSelectedCriteriaChanged(TSharedPtr<FListEntry> ListEntry, ESelectInfo::Type SelectionType)
+template <class FilterType>
+void SDataprepStringFilter<FilterType>::OnSelectedCriteriaChanged(TSharedPtr<FListEntry> ListEntry, ESelectInfo::Type SelectionType)
 {
 	check( Filter );
 	UEnum* Enum = StaticEnum< EDataprepStringMatchType >();
@@ -121,19 +123,22 @@ void SDataprepStringFilter::OnSelectedCriteriaChanged(TSharedPtr<FListEntry> Lis
 	}
 }
 
-FText SDataprepStringFilter::GetUserString() const
+template <class FilterType>
+FText SDataprepStringFilter<FilterType>::GetUserString() const
 {
 	check( Filter );
 	return FText::FromString( Filter->GetUserString() );
 }
 
-void SDataprepStringFilter::OnUserStringChanged(const FText& NewText)
+template <class FilterType>
+void SDataprepStringFilter<FilterType>::OnUserStringChanged(const FText& NewText)
 {
 	check( Filter );
 	Filter->SetUserString( NewText.ToString() );
 }
 
-void SDataprepStringFilter::OnUserStringComitted(const FText& NewText, ETextCommit::Type CommitType)
+template <class FilterType>
+void SDataprepStringFilter<FilterType>::OnUserStringComitted(const FText& NewText, ETextCommit::Type CommitType)
 {
 	check( Filter );
 	FString NewUserString = NewText.ToString();
@@ -147,9 +152,14 @@ void SDataprepStringFilter::OnUserStringComitted(const FText& NewText, ETextComm
 	}
 }
 
-void SDataprepStringFilter::AddReferencedObjects(FReferenceCollector& Collector)
+template <class FilterType>
+void SDataprepStringFilter<FilterType>::AddReferencedObjects(FReferenceCollector& Collector)
 {
 	Collector.AddReferencedObject( Filter );
 }
 
+// Explicit template instantiation
+template class SDataprepStringFilter<UDataprepStringFilter>;
+template class SDataprepStringFilter<UDataprepStringsArrayFilter>;
+
 #undef LOCTEXT_NAMESPACE
diff --git a/Engine/Plugins/Experimental/DataPrepEditor/Source/DataPrepEditor/Private/Widgets/Action/SDataprepStringFilter.h b/Engine/Plugins/Experimental/DataPrepEditor/Source/DataPrepEditor/Private/Widgets/Action/SDataprepStringFilter.h
--- a/Engine/Plugins/Experimental/DataPrepEditor/Source/DataPrepEditor/Private/Widgets/Action/SDataprepStringFilter.h
+++ b/Engine/Plugins/Experimental/DataPrepEditor/Source/DataPrepEditor/Private/Widgets/Action/SDataprepStringFilter.h
@@ -9,17 +9,17 @@
 #include "Widgets/SCompoundWidget.h"
 
 class SComboButton;
-class UDataprepStringFilter;
 
 template<class t>
 class SComboBox;
 
+template <class FilterType>
 class SDataprepStringFilter : public SCompoundWidget,  public FGCObject
 {
 	SLATE_BEGIN_ARGS(SDataprepStringFilter) {}
 	SLATE_END_ARGS();
 
-	void Construct(const FArguments& InArgs, UDataprepStringFilter& InFilter);
+	void Construct(const FArguments& InArgs, FilterType& InFilter);
 
 private:
 	// The string matching option for the combo box ( Displayed text, Tooltip, mapping for the UEnum)
@@ -43,7 +43,7 @@ private:
 
 	FString OldUserString;
 
-	UDataprepStringFilter* Filter;
+	FilterType* Filter;
 
 	TArray<TSharedPtr<FListEntry>> StringMatchingOptions;
 
diff --git a/Engine/Plugins/Experimental/DataPrepEditor/Source/DataPrepLibraries/Private/Fetchers/DataprepBoolFetcherLibrary.cpp b/Engine/Plugins/Experimental/DataPrepEditor/Source/DataPrepLibraries/Private/Fetchers/DataprepBoolFetcherLibrary.cpp
--- a/Engine/Plugins/Experimental/DataPrepEditor/Source/DataPrepLibraries/Private/Fetchers/DataprepBoolFetcherLibrary.cpp
+++ b/Engine/Plugins/Experimental/DataPrepEditor/Source/DataPrepLibraries/Private/Fetchers/DataprepBoolFetcherLibrary.cpp
@@ -4,6 +4,8 @@
 
 #include "GameFramework/Actor.h"
 
+#define LOCTEXT_NAMESPACE "DataprepBoolFetcherLibrary"
+
 /* UDataprepIsOfClassFetcher methods
  *****************************************************************************/
 FText UDataprepIsClassOfFetcher::AdditionalKeyword = NSLOCTEXT("DataprepIsClassOfFetcher","AdditionalKeyword","Is Child Of");
@@ -37,21 +39,9 @@ FText UDataprepIsClassOfFetcher::GetAdditionalKeyword_Implementation() const
 	return AdditionalKeyword;
 }
 
-
-/* UDataprepHasActorTagFetcher methods
- *****************************************************************************/
-bool UDataprepHasActorTagFetcher::Fetch_Implementation(const UObject* Object, bool& bOutFetchSucceded) const
+FText UDataprepIsClassOfFetcher::GetNodeDisplayFetcherName_Implementation() const
 {
-	if ( const AActor* Actor = Cast< const AActor >( Object ) )
-	{
-		bOutFetchSucceded = true;
-		return Actor->Tags.Contains( Tag );
-	}
-	 bOutFetchSucceded = false;
-	return false;
+	return LOCTEXT("ClassFilterTitle", "Class");
 }
 
-bool UDataprepHasActorTagFetcher::IsThreadSafe() const
-{
-	return true;
-}
+#undef LOCTEXT_NAMESPACE
diff --git a/Engine/Plugins/Experimental/DataPrepEditor/Source/DataPrepLibraries/Private/Fetchers/DataprepFloatFetcherLibrary.cpp b/Engine/Plugins/Experimental/DataPrepEditor/Source/DataPrepLibraries/Private/Fetchers/DataprepFloatFetcherLibrary.cpp
--- a/Engine/Plugins/Experimental/DataPrepEditor/Source/DataPrepLibraries/Private/Fetchers/DataprepFloatFetcherLibrary.cpp
+++ b/Engine/Plugins/Experimental/DataPrepEditor/Source/DataPrepLibraries/Private/Fetchers/DataprepFloatFetcherLibrary.cpp
@@ -9,6 +9,8 @@
 #include "Misc/Optional.h"
 #include "UObject/Object.h"
 
+#define LOCTEXT_NAMESPACE "DataprepFloatFetcherLibrary"
+
 /* UDataprepFloatVolumeFetcher methods
  *****************************************************************************/
 float UDataprepFloatBoundingVolumeFetcher::Fetch_Implementation(const UObject* Object, bool& bOutFetchSucceded) const
@@ -85,3 +87,10 @@ bool UDataprepFloatBoundingVolumeFetcher::IsThreadSafe() const
 {
 	return true;
 }
+
+FText UDataprepFloatBoundingVolumeFetcher::GetNodeDisplayFetcherName_Implementation() const
+{
+	return LOCTEXT("BoundingVolumeFilterTitle", "Bounding Volume");
+}
+
+#undef LOCTEXT_NAMESPACE
diff --git a/Engine/Plugins/Experimental/DataPrepEditor/Source/DataPrepLibraries/Private/Fetchers/DataprepStringFetcherLibrary.cpp b/Engine/Plugins/Experimental/DataPrepEditor/Source/DataPrepLibraries/Private/Fetchers/DataprepStringFetcherLibrary.cpp
--- a/Engine/Plugins/Experimental/DataPrepEditor/Source/DataPrepLibraries/Private/Fetchers/DataprepStringFetcherLibrary.cpp
+++ b/Engine/Plugins/Experimental/DataPrepEditor/Source/DataPrepLibraries/Private/Fetchers/DataprepStringFetcherLibrary.cpp
@@ -5,6 +5,8 @@
 #include "GameFramework/Actor.h"
 #include "UObject/Object.h"
 
+#define LOCTEXT_NAMESPACE "DataprepStringFetcherLibrary"
+
 /* UDataprepStringObjectNameFetcher methods
  *****************************************************************************/
 FString UDataprepStringObjectNameFetcher::Fetch_Implementation(const UObject* Object, bool& bOutFetchSucceded) const
@@ -24,6 +26,10 @@ bool UDataprepStringObjectNameFetcher::IsThreadSafe() const
 	return true;
 }
 
+FText UDataprepStringObjectNameFetcher::GetNodeDisplayFetcherName_Implementation() const
+{
+	return LOCTEXT("ObjectNameFilterTitle", "Object Name");
+}
 
 /* UDataprepStringActorLabelFetcher methods
  *****************************************************************************/
@@ -43,3 +49,10 @@ bool UDataprepStringActorLabelFetcher::IsThreadSafe() const
 {
 	return true;
 }
+
+FText UDataprepStringActorLabelFetcher::GetNodeDisplayFetcherName_Implementation() const
+{
+	return LOCTEXT("ActorLabelFilterTitle", "Actor Label");
+}
+
+#undef LOCTEXT_NAMESPACE
diff --git a/Engine/Plugins/Experimental/DataPrepEditor/Source/DataPrepLibraries/Private/Fetchers/DataprepStringsArrayFetcherLibrary.cpp b/Engine/Plugins/Experimental/DataPrepEditor/Source/DataPrepLibraries/Private/Fetchers/DataprepStringsArrayFetcherLibrary.cpp
new file mode 100644
--- /dev/null
+++ b/Engine/Plugins/Experimental/DataPrepEditor/Source/DataPrepLibraries/Private/Fetchers/DataprepStringsArrayFetcherLibrary.cpp
@@ -0,0 +1,40 @@
+// Copyright 1998-2019 Epic Games, Inc. All Rights Reserved.
+
+#include "Fetchers/DataprepStringsArrayFetcherLibrary.h"
+
+#include "GameFramework/Actor.h"
+#include "UObject/Object.h"
+
+#define LOCTEXT_NAMESPACE "DataprepStringsArrayFetcherLibrary"
+
+/* UDataprepStringActorTagsFetcher methods
+ *****************************************************************************/
+TArray<FString> UDataprepStringActorTagsFetcher::Fetch_Implementation(const UObject* Object, bool& bOutFetchSucceded) const
+{
+	if (const AActor* Actor = Cast<const AActor>(Object))
+	{
+		bOutFetchSucceded = true;
+		TArray<FString> Tags;
+		Tags.Empty(Actor->Tags.Num());
+		for (int TagIndex = 0; TagIndex < Actor->Tags.Num(); ++TagIndex)
+		{
+			Tags.Add(Actor->Tags[TagIndex].ToString());
+		}
+		return MoveTemp(Tags);
+	}
+
+	bOutFetchSucceded = false;
+	return {};
+}
+
+bool UDataprepStringActorTagsFetcher::IsThreadSafe() const
+{
+	return true;
+}
+
+FText UDataprepStringActorTagsFetcher::GetNodeDisplayFetcherName_Implementation() const
+{
+	return LOCTEXT("ActorTagsFilterTitle", "Tag");
+}
+
+#undef LOCTEXT_NAMESPACE
diff --git a/Engine/Plugins/Experimental/DataPrepEditor/Source/DataPrepLibraries/Public/Fetchers/DataprepBoolFetcherLibrary.h b/Engine/Plugins/Experimental/DataPrepEditor/Source/DataPrepLibraries/Public/Fetchers/DataprepBoolFetcherLibrary.h
--- a/Engine/Plugins/Experimental/DataPrepEditor/Source/DataPrepLibraries/Public/Fetchers/DataprepBoolFetcherLibrary.h
+++ b/Engine/Plugins/Experimental/DataPrepEditor/Source/DataPrepLibraries/Public/Fetchers/DataprepBoolFetcherLibrary.h
@@ -22,6 +22,7 @@ public:
 	//~ UDataprepFetcher interface
 	virtual bool IsThreadSafe() const final;
 	virtual FText GetAdditionalKeyword_Implementation() const;
+	virtual FText GetNodeDisplayFetcherName_Implementation() const;
 	//~ End of UDataprepFetcher interface
 
 	// The key for the for the string
@@ -33,20 +34,3 @@ public:
 
 	static FText AdditionalKeyword;
 };
-
-UCLASS(BlueprintType, NotBlueprintable, Meta = (DisplayName = "Has Actor Tag", ToolTip = "Return true when a actor has the specified tag."))
-class UDataprepHasActorTagFetcher final : public UDataprepBoolFetcher
-{
-	GENERATED_BODY()
-public:
-	//~ UDataprepBoolFetcher interface
-	virtual bool Fetch_Implementation(const UObject* Object, bool& bOutFetchSucceded) const final;
-	//~ End of UDataprepFloatFetcher interface
-
-	//~ UDataprepFetcher interface
-	virtual bool IsThreadSafe() const final;
-	//~ End of UDataprepFetcher interface
-
-	UPROPERTY(EditAnywhere, BlueprintReadWrite, Category = "Settings")
-	FName Tag;
-};
diff --git a/Engine/Plugins/Experimental/DataPrepEditor/Source/DataPrepLibraries/Public/Fetchers/DataprepFloatFetcherLibrary.h b/Engine/Plugins/Experimental/DataPrepEditor/Source/DataPrepLibraries/Public/Fetchers/DataprepFloatFetcherLibrary.h
--- a/Engine/Plugins/Experimental/DataPrepEditor/Source/DataPrepLibraries/Public/Fetchers/DataprepFloatFetcherLibrary.h
+++ b/Engine/Plugins/Experimental/DataPrepEditor/Source/DataPrepLibraries/Public/Fetchers/DataprepFloatFetcherLibrary.h
@@ -19,5 +19,6 @@ public:
 
 	//~ UDataprepFetcher interface
 	virtual bool IsThreadSafe() const final;
+	virtual FText GetNodeDisplayFetcherName_Implementation() const;
 	//~ End of UDataprepFetcher interface
 };
diff --git a/Engine/Plugins/Experimental/DataPrepEditor/Source/DataPrepLibraries/Public/Fetchers/DataprepStringFetcherLibrary.h b/Engine/Plugins/Experimental/DataPrepEditor/Source/DataPrepLibraries/Public/Fetchers/DataprepStringFetcherLibrary.h
--- a/Engine/Plugins/Experimental/DataPrepEditor/Source/DataPrepLibraries/Public/Fetchers/DataprepStringFetcherLibrary.h
+++ b/Engine/Plugins/Experimental/DataPrepEditor/Source/DataPrepLibraries/Public/Fetchers/DataprepStringFetcherLibrary.h
@@ -8,7 +8,6 @@
 
 #include "DataprepStringFetcherLibrary.generated.h"
 
-
 UCLASS(BlueprintType, NotBlueprintable, Meta = (DisplayName="Object Name", ToolTip="Get the name of an object."))
 class UDataprepStringObjectNameFetcher final : public UDataprepStringFetcher
 {
@@ -20,6 +19,7 @@ public:
 
 	//~ UDataprepFetcher interface
 	virtual bool IsThreadSafe() const final;
+	virtual FText GetNodeDisplayFetcherName_Implementation() const;
 	//~ End of UDataprepFetcher interface
 };
 
@@ -35,5 +35,6 @@ public:
 
 	//~ UDataprepFetcher interface
 	virtual bool IsThreadSafe() const final;
+	virtual FText GetNodeDisplayFetcherName_Implementation() const;
 	//~ End of UDataprepFetcher interface
 };
diff --git a/Engine/Plugins/Experimental/DataPrepEditor/Source/DataPrepLibraries/Public/Fetchers/DataprepStringsArrayFetcherLibrary.h b/Engine/Plugins/Experimental/DataPrepEditor/Source/DataPrepLibraries/Public/Fetchers/DataprepStringsArrayFetcherLibrary.h
new file mode 100644
--- /dev/null
+++ b/Engine/Plugins/Experimental/DataPrepEditor/Source/DataPrepLibraries/Public/Fetchers/DataprepStringsArrayFetcherLibrary.h
@@ -0,0 +1,25 @@
+// Copyright 1998-2019 Epic Games, Inc. All Rights Reserved.
+
+#pragma once
+
+#include "SelectionSystem/DataprepStringsArrayFetcher.h"
+
+#include "CoreMinimal.h"
+
+#include "DataprepStringsArrayFetcherLibrary.generated.h"
+
+UCLASS(BlueprintType, NotBlueprintable, Meta = (DisplayName = "Actor Tag", ToolTip = "Return the tags of an actor."))
+class UDataprepStringActorTagsFetcher final : public UDataprepStringsArrayFetcher
+{
+	GENERATED_BODY()
+
+public:
+	//~ UDataprepStringArrayFetcher interface
+	virtual TArray<FString> Fetch_Implementation(const UObject* Object, bool& bOutFetchSucceded) const override;
+	//~ End of UDataprepStringFetcher interface
+
+	//~ UDataprepFetcher interface
+	virtual bool IsThreadSafe() const final;
+	virtual FText GetNodeDisplayFetcherName_Implementation() const;
+	//~ End of UDataprepFetcher interface
+};
