Merging CL #9620516 and #9620528 from //UE4/Dev-Enterprise

9620516:
Add cancel button to progress report on Import & Execute
#jria UEENT-3009

9620528:
Fixed the issue identified by PVS-Studio

diff --git a/Engine/Plugins/Enterprise/DatasmithImporter/Source/DatasmithImporter/Private/DatasmithFileProducer.cpp b/Engine/Plugins/Enterprise/DatasmithImporter/Source/DatasmithImporter/Private/DatasmithFileProducer.cpp
--- a/Engine/Plugins/Enterprise/DatasmithImporter/Source/DatasmithImporter/Private/DatasmithFileProducer.cpp
+++ b/Engine/Plugins/Enterprise/DatasmithImporter/Source/DatasmithImporter/Private/DatasmithFileProducer.cpp
@@ -165,7 +165,7 @@ bool UDatasmithFileProducer::Initialize()
 	TSharedRef< IDatasmithScene > SceneElement = FDatasmithSceneFactory::CreateScene( *Source.GetSceneName() );
 
 	constexpr EObjectFlags LocalObjectFlags = RF_Public | RF_Standalone | RF_Transactional;
-	if ( !ImportContextPtr->Init( SceneElement, RootPath, LocalObjectFlags, GWarn, TSharedPtr< FJsonObject >(), true ) )
+	if ( !ImportContextPtr->Init( SceneElement, RootPath, LocalObjectFlags, Context.ProgressReporterPtr->GetFeedbackContext(), TSharedPtr< FJsonObject >(), true ) )
 	{
 		LogError( LOCTEXT( "DatasmithFileProducer_Initialization", "Initialization of producer failed." ) );
 		return false;
@@ -188,15 +188,25 @@ bool UDatasmithFileProducer::Execute(TArray< TWeakObjectPtr< UObject > >& OutAss
 		return false;
 	}
 
+	if ( IsCancelled() )
+	{
+		return false;
+	}
+
 	ProgressTaskPtr->ReportNextStep( FText::Format( LOCTEXT( "DatasmithFileProducer_ConvertingFile", "Converting {0} ..."), FText::FromString( FilePath ) ), 2.0f );
 	SceneElementToWorld();
 
+	if ( IsCancelled() )
+	{
+		return false;
+	}
+
 	ProgressTaskPtr->ReportNextStep( LOCTEXT( "DatasmithFileProducer_CleaningData", "Cleaning data ...") );
 	PreventNameCollision();
 
 	OutAssets.Append( MoveTemp( Assets ) );
 
-	return true;
+	return !IsCancelled();
 }
 
 // Borrowed from DatasmithImportFactoryImpl::ImportDatasmithScene
@@ -674,18 +684,23 @@ bool UDatasmithDirProducer::Execute(TArray< TWeakObjectPtr< UObject > >& OutAsse
 
 	for( const FString& FileName : FilesToProcess )
 	{
+		if ( IsCancelled() )
+		{
+			break;
+		}
+
 		FileProducer->FilePath =  FPaths::ConvertRelativePathToFull( FileName );
 
 		Task.ReportNextStep( FText::Format( LOCTEXT( "DatasmithFileProducer_LoadingFile", "Loading {0} ..."), FText::FromString( FileName ) ) );
 
-		if(!FileProducer->Produce(Context, OutAssets))
+		if( !FileProducer->Produce( Context, OutAssets ) )
 		{
 			FText ErrorReport = FText::Format( LOCTEXT( "DatasmithDirProducer_Failed", "Failed to load {0} ..."), FText::FromString( FileName ) );
 			LogError( ErrorReport );
 		}
 	}
 
-	return true;
+	return !IsCancelled();
 }
 
 void UDatasmithDirProducer::Reset()
diff --git a/Engine/Plugins/Experimental/DataPrepEditor/Source/DataPrepCore/Private/DataPrepAssetProducers.cpp b/Engine/Plugins/Experimental/DataPrepEditor/Source/DataPrepCore/Private/DataPrepAssetProducers.cpp
--- a/Engine/Plugins/Experimental/DataPrepEditor/Source/DataPrepCore/Private/DataPrepAssetProducers.cpp
+++ b/Engine/Plugins/Experimental/DataPrepEditor/Source/DataPrepCore/Private/DataPrepAssetProducers.cpp
@@ -39,7 +39,7 @@ void UDataprepAssetProducers::PostEditUndo()
 	OnChanged.Broadcast( FDataprepAssetChangeType::ProducerModified, INDEX_NONE );
 }
 
-TArray< TWeakObjectPtr< UObject > > UDataprepAssetProducers::Produce(const FDataprepProducerContext& InContext )
+TArray< TWeakObjectPtr< UObject > > UDataprepAssetProducers::Produce( const FDataprepProducerContext& InContext )
 {
 	if( AssetProducers.Num() == 0 )
 	{
@@ -65,8 +65,17 @@ TArray< TWeakObjectPtr< UObject > > UDataprepAssetProducers::Produce(const FData
 			{
 				if( !Producer->Produce( InContext, OutAssets ) )
 				{
-					FText Report = FText::Format(LOCTEXT( "ProducerReport_Failed", "Producer {0} failed to execute."), FText::FromString( Producer->GetName() ) );
-					InContext.LoggerPtr->LogError( Report, *Producer );
+					if ( Task.IsWorkCancelled() )
+					{
+						FText Report = FText::Format(LOCTEXT( "ProducerReport_Cancelled", "Producer {0} was cancelled during execution."), FText::FromString( Producer->GetName() ) );
+						InContext.LoggerPtr->LogInfo( Report, *Producer );
+						break;
+					}
+					else
+					{
+						FText Report = FText::Format(LOCTEXT( "ProducerReport_Failed", "Producer {0} failed to execute."), FText::FromString( Producer->GetName() ) );
+						InContext.LoggerPtr->LogError( Report, *Producer );
+					}
 				}
 			}
 		}
diff --git a/Engine/Plugins/Experimental/DataPrepEditor/Source/DataPrepCore/Private/DataprepActionAsset.cpp b/Engine/Plugins/Experimental/DataPrepEditor/Source/DataPrepCore/Private/DataprepActionAsset.cpp
--- a/Engine/Plugins/Experimental/DataPrepEditor/Source/DataPrepCore/Private/DataprepActionAsset.cpp
+++ b/Engine/Plugins/Experimental/DataPrepEditor/Source/DataPrepCore/Private/DataprepActionAsset.cpp
@@ -417,6 +417,11 @@ void UDataprepActionAsset::ExecuteAction(const TSharedPtr<FDataprepActionContext
 					ExecuteOneStep( Step );
 				}
 
+				if ( ContextPtr->ProgressReporterPtr->IsWorkCancelled() )
+				{
+					break;
+				}
+
 				if(ContextPtr->ContinueCallback && !ContextPtr->ContinueCallback(this, Step->Operation, Step->Filter))
 				{
 					break;
diff --git a/Engine/Plugins/Experimental/DataPrepEditor/Source/DataPrepCore/Private/DataprepAssetInterface.cpp b/Engine/Plugins/Experimental/DataPrepEditor/Source/DataPrepCore/Private/DataprepAssetInterface.cpp
--- a/Engine/Plugins/Experimental/DataPrepEditor/Source/DataPrepCore/Private/DataprepAssetInterface.cpp
+++ b/Engine/Plugins/Experimental/DataPrepEditor/Source/DataPrepCore/Private/DataprepAssetInterface.cpp
@@ -163,7 +163,7 @@ void UDataprepAssetInterface::ExecuteRecipe_Internal(const TSharedPtr<FDataprepA
 			Task.ReportNextStep(FText::Format(LOCTEXT("ExecutingAction", "Executing \"{0}\" ..."), FText::FromString(ActionAsset->GetLabel())));
 
 			ActionAsset->ExecuteAction(InActionsContext);
-			if (ActionAsset->bExecutionInterrupted)
+			if (Task.IsWorkCancelled())
 			{
 				Task.ReportNextStep(LOCTEXT("InterruptedExecution", "Execution interrupted ..."));
 				break;
diff --git a/Engine/Plugins/Experimental/DataPrepEditor/Source/DataPrepCore/Private/DataprepCoreLibrary.cpp b/Engine/Plugins/Experimental/DataPrepEditor/Source/DataPrepCore/Private/DataprepCoreLibrary.cpp
--- a/Engine/Plugins/Experimental/DataPrepEditor/Source/DataPrepCore/Private/DataprepCoreLibrary.cpp
+++ b/Engine/Plugins/Experimental/DataPrepEditor/Source/DataPrepCore/Private/DataprepCoreLibrary.cpp
@@ -119,7 +119,7 @@ bool UDataprepCoreLibrary::ExecuteWithReporting( UDataprepAssetInterface* Datapr
 #endif
 
 	TSharedPtr<IDataprepLogger> Logger( new FDataprepCoreUtils::FDataprepLogger );
-	TSharedPtr<IDataprepProgressReporter> Reporter( new FDataprepCoreUtils::FDataprepProgressUIReporter );
+	TSharedPtr<IDataprepProgressReporter> Reporter( new FDataprepCoreUtils::FDataprepProgressUIReporter() );
 	return Execute_Internal( DataprepAssetInterface, Logger, Reporter );
 }
 
@@ -162,12 +162,14 @@ bool UDataprepCoreLibrary::Execute_Internal(UDataprepAssetInterface* DataprepAss
 			UPackage* TransientPackage = NewObject< UPackage >( nullptr, *TransientContentFolder, RF_Transient );
 			TransientPackage->FullyLoad();
 
+			TSharedPtr< FDataprepCoreUtils::FDataprepFeedbackContext > FeedbackContext( new FDataprepCoreUtils::FDataprepFeedbackContext );
+
 			FDataprepProducerContext Context;
 			Context.SetWorld( TransientWorld.Get() )
 				.SetRootPackage( TransientPackage )
 				.SetLogger( Logger)
 				.SetProgressReporter( Reporter );
-
+			
 			FText Message = FText::Format( LOCTEXT( "Running_Producers", "Running \"{0}\'s Producers ..."), DataprepAssetTextName );
 			ProgressTask.ReportNextStep( Message );
 			Assets = DataprepAssetInterface->GetProducers()->Produce( Context );
diff --git a/Engine/Plugins/Experimental/DataPrepEditor/Source/DataPrepCore/Private/DataprepCoreUtils.cpp b/Engine/Plugins/Experimental/DataPrepEditor/Source/DataPrepCore/Private/DataprepCoreUtils.cpp
--- a/Engine/Plugins/Experimental/DataPrepEditor/Source/DataPrepCore/Private/DataprepCoreUtils.cpp
+++ b/Engine/Plugins/Experimental/DataPrepEditor/Source/DataPrepCore/Private/DataprepCoreUtils.cpp
@@ -178,6 +178,15 @@ void FDataprepWorkReporter::ReportNextStep(const FText & InMessage, float InIncr
 	}
 }
 
+bool FDataprepWorkReporter::IsWorkCancelled() const
+{
+	if ( Reporter.IsValid() )
+	{
+		return Reporter->IsWorkCancelled();
+	}
+	return false;
+}
+
 #ifdef NEW_DATASMITHSCENE_WORKFLOW
 void FDataprepCoreUtils::AddDataprepAssetUserData(UObject* Target, UDataprepAssetInterface* DataprepAssetInterface)
 {
@@ -227,8 +236,8 @@ void FDataprepCoreUtils::FDataprepLogger::LogError(const FText& InLogText,  cons
 
 void FDataprepCoreUtils::FDataprepProgressUIReporter::BeginWork( const FText& InTitle, float InAmountOfWork )
 {
-	ProgressTasks.Emplace( new FScopedSlowTask( InAmountOfWork, InTitle, true, *GWarn ) );
-	ProgressTasks.Last()->MakeDialog(false);
+	ProgressTasks.Emplace( new FScopedSlowTask( InAmountOfWork, InTitle, true, FeedbackContext.IsValid() ? *FeedbackContext.Get() : *GWarn ) );
+	ProgressTasks.Last()->MakeDialog(true);
 }
 
 void FDataprepCoreUtils::FDataprepProgressUIReporter::EndWork()
@@ -248,6 +257,21 @@ void FDataprepCoreUtils::FDataprepProgressUIReporter::ReportProgress( float Prog
 	}
 }
 
+bool FDataprepCoreUtils::FDataprepProgressUIReporter::IsWorkCancelled()
+{
+	if( !bIsCancelled && ProgressTasks.Num() > 0 )
+	{
+		const TSharedPtr<FScopedSlowTask>& ProgressTask = ProgressTasks.Last();
+		bIsCancelled |= ProgressTask->ShouldCancel();
+	}
+	return bIsCancelled;
+}
+
+FFeedbackContext* FDataprepCoreUtils::FDataprepProgressUIReporter::GetFeedbackContext() const
+{
+	return FeedbackContext.IsValid() ? FeedbackContext.Get() : GWarn;
+}
+
 void FDataprepCoreUtils::FDataprepProgressTextReporter::BeginWork( const FText& InTitle, float InAmountOfWork )
 {
 	UE_LOG( LogDataprepCore, Log, TEXT("Start: %s ..."), *InTitle.ToString() );
@@ -270,4 +294,14 @@ void FDataprepCoreUtils::FDataprepProgressTextReporter::ReportProgress( float Pr
 	}
 }
 
+bool FDataprepCoreUtils::FDataprepProgressTextReporter::IsWorkCancelled()
+{
+	return false;
+}
+
+FFeedbackContext* FDataprepCoreUtils::FDataprepProgressTextReporter::GetFeedbackContext() const
+{
+	return FeedbackContext.Get();
+}
+
 #undef LOCTEXT_NAMESPACE
\ No newline at end of file
diff --git a/Engine/Plugins/Experimental/DataPrepEditor/Source/DataPrepCore/Public/DataPrepContentProducer.h b/Engine/Plugins/Experimental/DataPrepEditor/Source/DataPrepCore/Public/DataPrepContentProducer.h
--- a/Engine/Plugins/Experimental/DataPrepEditor/Source/DataPrepCore/Public/DataPrepContentProducer.h
+++ b/Engine/Plugins/Experimental/DataPrepEditor/Source/DataPrepCore/Public/DataPrepContentProducer.h
@@ -11,6 +11,8 @@
 
 #include "DataPrepContentProducer.generated.h"
 
+class FFeedbackContext;
+
 /** Structure to pass execution context to producer */
 struct FDataprepProducerContext
 {
@@ -112,6 +114,9 @@ public:
 	 */
 	virtual bool Supersede(const UDataprepContentProducer* OtherProducer) const { unimplemented(); return true; }
 
+	/** Returns true if the producer was cancelled during execution */
+	bool IsCancelled() { return Context.ProgressReporterPtr.IsValid() ? Context.ProgressReporterPtr->IsWorkCancelled() : false; }
+
 protected:
 
 	/**
@@ -168,6 +173,9 @@ protected:
 	/** Delegate to broadcast changes to the producer */
 	FDataprepProducerChanged OnChanged;
 
+	/** The feedback context in use to check for user cancellation */
+	TSharedPtr<FFeedbackContext> FeedbackContext;
+
 private:
 	void Terminate()
 	{
diff --git a/Engine/Plugins/Experimental/DataPrepEditor/Source/DataPrepCore/Public/DataprepCoreUtils.h b/Engine/Plugins/Experimental/DataPrepEditor/Source/DataPrepCore/Public/DataprepCoreUtils.h
--- a/Engine/Plugins/Experimental/DataPrepEditor/Source/DataPrepCore/Public/DataprepCoreUtils.h
+++ b/Engine/Plugins/Experimental/DataPrepEditor/Source/DataPrepCore/Public/DataprepCoreUtils.h
@@ -7,6 +7,9 @@
 #include "IDataprepLogger.h"
 #include "IDataprepProgressReporter.h"
 
+#include "FeedbackContextEditor.h"
+#include "HAL/FeedbackContextAnsi.h"
+
 #include "UObject/Object.h"
 #include "UObject/UObjectGlobals.h"
 #include "UObject/Package.h"
@@ -79,10 +82,38 @@ public:
 
 	};
 
+	class DATAPREPCORE_API FDataprepFeedbackContext : public FFeedbackContextEditor
+	{
+	public:
+		/** 
+		 * We want to override this method in order to cache the cancel result and not clear it,
+		 * so it can be checked multiple times with the correct result!
+		 * (FFeedbackContextEditor::ReceivedUserCancel clears it)
+		 */
+		virtual bool ReceivedUserCancel() override 
+		{ 
+			if ( !bTaskWasCancelledCache )
+			{
+				bTaskWasCancelledCache = FFeedbackContextEditor::ReceivedUserCancel();
+			}
+			return bTaskWasCancelledCache;
+		}
+
+	private:
+		bool bTaskWasCancelledCache = false;
+	};
+
 	class DATAPREPCORE_API FDataprepProgressUIReporter : public IDataprepProgressReporter
 	{
 	public:
 		FDataprepProgressUIReporter()
+			: bIsCancelled(false)
+		{
+		}
+
+		FDataprepProgressUIReporter( TSharedRef<FFeedbackContext> InFeedbackContext )
+			: FeedbackContext(InFeedbackContext)
+			, bIsCancelled(false)
 		{
 		}
 
@@ -94,10 +125,14 @@ public:
 		virtual void BeginWork( const FText& InTitle, float InAmountOfWork ) override;
 		virtual void EndWork() override;
 		virtual void ReportProgress( float Progress, const FText& InMessage ) override;
+		virtual bool IsWorkCancelled() override;
+		virtual FFeedbackContext* GetFeedbackContext() const override;
 		// End IDataprepProgressReporter interface
 
 	private:
 		TArray< TSharedPtr< FScopedSlowTask > > ProgressTasks;
+		TSharedPtr< FFeedbackContext > FeedbackContext;
+		bool bIsCancelled;
 	};
 
 	class DATAPREPCORE_API FDataprepProgressTextReporter : public IDataprepProgressReporter
@@ -105,6 +140,7 @@ public:
 	public:
 		FDataprepProgressTextReporter()
 			: TaskDepth(0)
+			, FeedbackContext( new FFeedbackContextAnsi )
 		{
 		}
 
@@ -116,9 +152,12 @@ public:
 		virtual void BeginWork( const FText& InTitle, float InAmountOfWork ) override;
 		virtual void EndWork() override;
 		virtual void ReportProgress( float Progress, const FText& InMessage ) override;
+		virtual bool IsWorkCancelled() override;
+		virtual FFeedbackContext* GetFeedbackContext() const override;
 
 	private:
 		int32 TaskDepth;
+		TUniquePtr<FFeedbackContextAnsi> FeedbackContext;
 	};
 };
 
diff --git a/Engine/Plugins/Experimental/DataPrepEditor/Source/DataPrepCore/Public/IDataprepProgressReporter.h b/Engine/Plugins/Experimental/DataPrepEditor/Source/DataPrepCore/Public/IDataprepProgressReporter.h
--- a/Engine/Plugins/Experimental/DataPrepEditor/Source/DataPrepCore/Public/IDataprepProgressReporter.h
+++ b/Engine/Plugins/Experimental/DataPrepEditor/Source/DataPrepCore/Public/IDataprepProgressReporter.h
@@ -6,6 +6,7 @@
 
 class FText;
 class UObject;
+class FFeedbackContext;
 
 /**
  * This is the interface that a progress reporter must implement to work with FDataprepWorkReporter
@@ -33,6 +34,12 @@ public:
 	 */
 	virtual void ReportProgress( float IncrementOfWork, const FText& InMessage ) = 0;
 
+	/** Returns true if the work has been cancelled */
+	virtual bool IsWorkCancelled() = 0;
+
+	/** Returns the feedback context used by this progress reporter */
+	virtual FFeedbackContext* GetFeedbackContext() const = 0;
+
 	friend class FDataprepWorkReporter;
 };
 
@@ -69,6 +76,8 @@ public:
 		ReportNextStep( InMessage, DefaultIncrementOfWork );
 	}
 
+	bool IsWorkCancelled() const;
+
 private:
 	/** Dataprep progress reporter associated with the task */
 	TSharedPtr<IDataprepProgressReporter> Reporter;
diff --git a/Engine/Plugins/Experimental/DataPrepEditor/Source/DataPrepEditor/Private/DataPrepEditor.cpp b/Engine/Plugins/Experimental/DataPrepEditor/Source/DataPrepEditor/Private/DataPrepEditor.cpp
--- a/Engine/Plugins/Experimental/DataPrepEditor/Source/DataPrepEditor/Private/DataPrepEditor.cpp
+++ b/Engine/Plugins/Experimental/DataPrepEditor/Source/DataPrepEditor/Private/DataPrepEditor.cpp
@@ -446,7 +446,6 @@ void FDataprepEditor::InitDataprepEditor(const EToolkitMode::Type Mode, const TS
 
 	ActionsContext->SetTransientContentFolder( GetTransientContentFolder() / TEXT("Pipeline") )
 		.SetLogger( TSharedPtr<IDataprepLogger>( new FDataprepCoreUtils::FDataprepLogger ) )
-		.SetProgressReporter( TSharedPtr< IDataprepProgressReporter >( new FDataprepCoreUtils::FDataprepProgressUIReporter() ) )
 		.SetCanExecuteNextStep( CanExecuteNextStepFunc )
 		.SetActionsContextChanged( ActionsContextChangedFunc );
 
@@ -569,6 +568,9 @@ void FDataprepEditor::OnBuildWorld()
 	UPackage* TransientPackage = NewObject< UPackage >( nullptr, *GetTransientContentFolder(), RF_Transient );
 	TransientPackage->FullyLoad();
 
+	TSharedPtr< FDataprepCoreUtils::FDataprepFeedbackContext > FeedbackContext( new FDataprepCoreUtils::FDataprepFeedbackContext );
+
+	TSharedPtr< IDataprepProgressReporter > ProgressReporter( new FDataprepCoreUtils::FDataprepProgressUIReporter( FeedbackContext.ToSharedRef() ) );
 	{
 		FTimeLogger TimeLogger( TEXT("Import") );
 
@@ -577,11 +579,18 @@ void FDataprepEditor::OnBuildWorld()
 		Context.SetWorld( PreviewWorld.Get() )
 			.SetRootPackage( TransientPackage )
 			.SetLogger( TSharedPtr< IDataprepLogger >( new FDataprepCoreUtils::FDataprepLogger ) )
-			.SetProgressReporter( TSharedPtr< IDataprepProgressReporter >( new FDataprepCoreUtils::FDataprepProgressUIReporter() ) );
+			.SetProgressReporter( ProgressReporter );
 
 		Assets = DataprepAssetInterface->GetProducers()->Produce( Context );
 	}
 
+	if ( ProgressReporter->IsWorkCancelled() )
+	{
+		// Flush the work that's already been done
+		ResetBuildWorld();
+		return;
+	}
+
 	CachedAssets.Reset();
 	CachedAssets.Append( Assets );
 
@@ -703,8 +712,6 @@ void FDataprepEditor::OnExecutePipeline()
 	// Required because assets could be removed and the 3D viewport is using raw pointers
 	SceneViewportView->ClearMeshes();
 
-	TSharedPtr< IDataprepProgressReporter > ProgressReporter( new FDataprepCoreUtils::FDataprepProgressUIReporter() );
-
 	// Trigger execution of data preparation operations on world attached to recipe
 	{
 		FTimeLogger TimeLogger( TEXT("ExecutePipeline") );
@@ -713,6 +720,8 @@ void FDataprepEditor::OnExecutePipeline()
 		// Editors can individually refuse this request: we ignore it during the pipeline traversal.
 		TGuardValue<bool> IgnoreCloseRequestGuard(bIgnoreCloseRequest, true);
 
+		TSharedPtr< FDataprepCoreUtils::FDataprepFeedbackContext > FeedbackContext(new FDataprepCoreUtils::FDataprepFeedbackContext);
+		ActionsContext->SetProgressReporter( TSharedPtr< IDataprepProgressReporter >( new FDataprepCoreUtils::FDataprepProgressUIReporter( FeedbackContext.ToSharedRef() ) ) );
 		ActionsContext->SetWorld( PreviewWorld.Get() ).SetAssets( Assets );
 
 		DataprepAssetInterfacePtr->ExecuteRecipe( ActionsContext );
@@ -721,6 +730,11 @@ void FDataprepEditor::OnExecutePipeline()
 		Assets = ActionsContext->Assets.Array();
 	}
 
+	if ( ActionsContext->ProgressReporterPtr->IsWorkCancelled() )
+	{
+		RestoreFromSnapshot();
+	}
+
 	// Redraw 3D viewport
 	SceneViewportView->UpdateMeshes();
 
