Merging CL #9542858 from //UE4/Dev-Enterprise

Fixed logic error for environment settings

diff --git a/Engine/Plugins/Experimental/DataPrepEditor/Source/DataPrepEditor/Private/Widgets/SDataprepEditorViewport.cpp b/Engine/Plugins/Experimental/DataPrepEditor/Source/DataPrepEditor/Private/Widgets/SDataprepEditorViewport.cpp
--- a/Engine/Plugins/Experimental/DataPrepEditor/Source/DataPrepEditor/Private/Widgets/SDataprepEditorViewport.cpp
+++ b/Engine/Plugins/Experimental/DataPrepEditor/Source/DataPrepEditor/Private/Widgets/SDataprepEditorViewport.cpp
@@ -1253,6 +1253,14 @@ void SDataprepEditorViewport::ApplyRenderingMaterial()
 
 void SDataprepEditorViewport::LoadDefaultSettings()
 {
+	// Disable viewing settings for the time being
+	static bool bAllowViewingSettings = false;
+	if(!bAllowViewingSettings)
+	{
+		AssetViewerProfileIndex = 0;
+		return;
+	}
+
 	// Find index of Dataprep's viewport's settings
 	const TCHAR* DataprepViewportSettingProfileName = TEXT("DataprepViewportSetting");
 
@@ -1295,7 +1303,7 @@ void SDataprepEditorViewport::LoadDefaultSettings()
 		if(EnvironmentCubeMapPath != DataprepViewportSettingProfile.EnvironmentCubeMapPath)
 		{
 			// Check that the Cube map does exist
-			FSoftObjectPath EnvironmentCubeMap( DataprepViewportSettingProfile.EnvironmentCubeMapPath );
+			FSoftObjectPath EnvironmentCubeMap( EnvironmentCubeMapPath );
 			UObject* LoadedObject = EnvironmentCubeMap.TryLoad();
 
 			while (UObjectRedirector* Redirector = Cast<UObjectRedirector>( LoadedObject ))
@@ -1307,7 +1315,7 @@ void SDataprepEditorViewport::LoadDefaultSettings()
 			if( Cast<UTextureCube>( LoadedObject ) != nullptr )
 			{
 				DataprepViewportSettingProfile.EnvironmentCubeMapPath = EnvironmentCubeMapPath;
-				DataprepViewportSettingProfile.EnvironmentCubeMap = EnvironmentCubeMap;
+				DataprepViewportSettingProfile.EnvironmentCubeMap = LoadedObject;
 			}
 		}
 	}
