Merging CL #9487349 from //UE4/Dev-Enterprise

Split "Clean World" operation into "Delete Unused Assets" and "Compact Scene Graph" operations.

diff --git a/Engine/Plugins/Experimental/DataPrepEditor/Source/DataPrepLibraries/Private/DataprepEditingOperations.cpp b/Engine/Plugins/Experimental/DataPrepEditor/Source/DataPrepLibraries/Private/DataprepEditingOperations.cpp
--- a/Engine/Plugins/Experimental/DataPrepEditor/Source/DataPrepLibraries/Private/DataprepEditingOperations.cpp
+++ b/Engine/Plugins/Experimental/DataPrepEditor/Source/DataPrepLibraries/Private/DataprepEditingOperations.cpp
@@ -85,7 +85,7 @@ namespace DatasmithEditingOperationsUtils
 	};
 }
 
-void UDataprepRemoveObjectsOperation::OnExecution_Implementation(const FDataprepContext& InContext)
+void UDataprepDeleteObjectsOperation::OnExecution_Implementation(const FDataprepContext& InContext)
 {
 #ifdef LOG_TIME
 	DataprepEditingOperationTime::FTimeLogger TimeLogger( TEXT("RemoveObjects"), [&]( FText Text) { this->LogInfo( Text ); });
@@ -157,34 +157,34 @@ void UDataprepRemoveObjectsOperation::OnExecution_Implementation(const FDataprep
 
 void UDataprepMergeActorsOperation::OnExecution_Implementation(const FDataprepContext& InContext)
 {
-	TArray<AStaticMeshActor*> ActorsToMerge;
+	TArray<AActor*> ActorsToMerge;
 	TArray<UPrimitiveComponent*> ComponentsToMerge;
 	UWorld* CurrentWorld = nullptr;
 
 	for(UObject* Object : InContext.Objects)
 	{
-		if( AStaticMeshActor* MeshActor = Cast<AStaticMeshActor>(Object) )
+		if( AActor* Actor = Cast<AActor>(Object) )
 		{
-			if(!MeshActor->IsPendingKillOrUnreachable())
+			if(!Actor->IsPendingKillOrUnreachable())
 			{
 				if(CurrentWorld == nullptr)
 				{
-					CurrentWorld = MeshActor->GetWorld();
+					CurrentWorld = Actor->GetWorld();
 				}
 
-				if(CurrentWorld != MeshActor->GetWorld())
+				if(CurrentWorld != Actor->GetWorld())
 				{
 					// #ueent_todo: Warn that incompatible actor found and discarded
 					continue;
 				}
 
 				TInlineComponentArray<UStaticMeshComponent*> ComponentArray;
-				MeshActor->GetComponents<UStaticMeshComponent>(ComponentArray);
+				Actor->GetComponents<UStaticMeshComponent>(ComponentArray);
 
 				bool bMeshActorIsValid = false;
 				for (UStaticMeshComponent* MeshComponent : ComponentArray)
 				{
-					if (MeshComponent->GetStaticMesh()/* && MeshComponent->GetStaticMesh()->RenderData.IsValid()*/)
+					if (MeshComponent->GetStaticMesh())
 					{
 						bMeshActorIsValid = true;
 						ComponentsToMerge.Add(MeshComponent);
@@ -194,7 +194,7 @@ void UDataprepMergeActorsOperation::OnExecution_Implementation(const FDataprepCo
 				//Actor needs at least one StaticMeshComponent to be considered valid
 				if (bMeshActorIsValid)
 				{
-					ActorsToMerge.Add( MeshActor );
+					ActorsToMerge.Add( Actor );
 				}
 
 			}
@@ -228,31 +228,36 @@ void UDataprepMergeActorsOperation::OnExecution_Implementation(const FDataprepCo
 
 	MergedActor->GetRootComponent()->SetWorldLocation( MergedMeshWorldLocation );
 
+	// Keep the merged actor in the hierarchy, taking the parent of the first component
+	// In the future, the merged actor could be attached to the common ancestor instead of the first parent in the list
+	MergedActor->GetRootComponent()->AttachToComponent( ComponentsToMerge[0]->GetAttachParent(), FAttachmentTransformRules::KeepWorldTransform );
+
 	// Collect all objects to be deleted
 	TArray<UObject*> ObjectsToDelete;
 	ObjectsToDelete.Reserve( ComponentsToMerge.Num() + ActorsToMerge.Num() );
 
-	if(bDeleteMergedMeshes)
+	// Simple way to delete the actors: detach the merged components if it's safe to do so
+	for (UPrimitiveComponent* Component : ComponentsToMerge)
 	{
-		TSet<UObject*> StaticMeshes;
-		StaticMeshes.Reserve( ComponentsToMerge.Num() );
-
-		for(UPrimitiveComponent* PrimitiveComponent : ComponentsToMerge)
+		if (Component->GetNumChildrenComponents() == 0)
 		{
-			if(UStaticMesh* StaticMesh = Cast<UStaticMeshComponent>(PrimitiveComponent)->GetStaticMesh())
-			{
-				StaticMeshes.Add(StaticMesh);
-				ObjectsToDelete.Add( StaticMesh );
-				Cast<UStaticMeshComponent>(PrimitiveComponent)->SetStaticMesh( nullptr );
-			}
+			Component->DetachFromComponent(FDetachmentTransformRules::KeepRelativeTransform);
+		}
+		else if (UStaticMeshComponent* StaticMeshComponent = Cast<UStaticMeshComponent>(Component))
+		{
+			// To keep the children in the hierarchy, simply hide the StaticMesh, but the StaticMeshComponent could be replaced by a SceneComponent and
+			// then attach the children to it
+			StaticMeshComponent->SetStaticMesh(nullptr);
 		}
 	}
 
-	if(bDeleteMergedActors)
+	// Then delete the merged actors that don't have any children component
+	for(AActor* Actor : ActorsToMerge)
 	{
-		for(AStaticMeshActor* MeshActor : ActorsToMerge)
+		UPrimitiveComponent* RootComponent = Cast<UPrimitiveComponent>(Actor->GetRootComponent());
+		if (RootComponent && RootComponent->GetNumChildrenComponents() == 0)
 		{
-			ObjectsToDelete.Add( MeshActor );
+			ObjectsToDelete.Add(Actor);
 		}
 	}
 
@@ -261,11 +266,28 @@ void UDataprepMergeActorsOperation::OnExecution_Implementation(const FDataprepCo
 
 bool UDataprepMergeActorsOperation::MergeStaticMeshActors(UWorld* World, const TArray<UPrimitiveComponent*>& ComponentsToMerge, const FString& RootName, bool bCreateActor)
 {
+	TSet<UStaticMesh*> StaticMeshes;
+	for(UPrimitiveComponent* PrimitiveComponent : ComponentsToMerge)
+	{
+		if(UStaticMeshComponent* StaticMeshComponent = Cast<UStaticMeshComponent>(PrimitiveComponent))
+		{
+			if(StaticMeshComponent->GetStaticMesh()->RenderData == nullptr)
+			{
+				StaticMeshes.Add(StaticMeshComponent->GetStaticMesh());
+			}
+		}
+	}
+
+	FStaticMeshBuilder StaticMeshBuilder(StaticMeshes);
+
 	//
 	// See MeshMergingTool.cpp
 	//
 	const IMeshMergeUtilities& MeshUtilities = FModuleManager::Get().LoadModuleChecked<IMeshMergeModule>("MeshMergeUtilities").GetUtilities();
 
+	FMeshMergingSettings MergeSettings;
+	MergeSettings.bPivotPointAtZero = bPivotPointAtZero;
+
 	TArray<UObject*> CreatedAssets;
 	const float ScreenAreaSize = TNumericLimits<float>::Max();
 	MeshUtilities.MergeComponentsToStaticMesh( ComponentsToMerge, World, MergeSettings, nullptr, GetTransientPackage(), FString(), CreatedAssets, MergedMeshWorldLocation, ScreenAreaSize, true);
@@ -303,263 +325,47 @@ bool UDataprepMergeActorsOperation::MergeStaticMeshActors(UWorld* World, const T
 	return true;
 }
 
-void UDataprepMergeActorsOperation::PrepareStaticMeshes(TSet<UStaticMesh*> StaticMeshes, IMeshBuilderModule& MeshBuilderModule)
+UDataprepMergeActorsOperation::FStaticMeshBuilder::FStaticMeshBuilder(const TSet<UStaticMesh *>& InStaticMeshes)
 {
-	TRACE_CPUPROFILER_EVENT_SCOPE(UDataprepMergeActorsOperation::PrepareStaticMeshes);
+	StaticMeshes = InStaticMeshes;
 
-	if( StaticMeshes.Num() > 1 )
-	{
-		TArray<UStaticMesh*> StaticMeshesToBuild = StaticMeshes.Array();
-		ParallelFor( StaticMeshesToBuild.Num(), [&]( int32 Index ) {
-			DataprepOperationsLibraryUtil::BuildRenderData( StaticMeshesToBuild[Index], MeshBuilderModule );
-		});
-	}
-	else
-	{
-		for(UStaticMesh* StaticMesh : StaticMeshes)
-		{
-			DataprepOperationsLibraryUtil::BuildRenderData( StaticMesh, MeshBuilderModule );
-		}
-	}
+	IMeshBuilderModule& MeshBuilderModule = FModuleManager::LoadModuleChecked< IMeshBuilderModule >( TEXT("MeshBuilder") );
+	UDataprepMergeActorsOperation::PrepareStaticMeshes(StaticMeshes, MeshBuilderModule);
 }
 
-void UDataprepSmartMergeOperation::OnExecution_Implementation(const FDataprepContext& InContext)
+UDataprepMergeActorsOperation::FStaticMeshBuilder::~FStaticMeshBuilder()
 {
-	UWorld* World = nullptr;
-
-	for(UObject* Object : InContext.Objects)
+	// Release render data of built static meshes
+	for(UStaticMesh* StaticMesh : StaticMeshes)
 	{
-		if(AActor* Actor = Cast<AActor>(Object))
+		if(StaticMesh)
 		{
-			World = Actor->GetWorld();
-			break;
+			StaticMesh->RenderData.Reset();
 		}
 	}
-
-	if(World == nullptr/* || World == GWorld*/)
-	{
-		return;
-	}
-
-	// Get root actors
-	TArray<AActor*> RootActors;
-	DatasmithEditingOperationsUtils::GetRootActors(World, RootActors);
-
-	TArray<AActor*> Actors;
-	DatasmithEditingOperationsUtils::FindActorsToMerge(RootActors, Actors);
-
-#ifdef LOG_TIME
-	DataprepEditingOperationTime::FTimeLogger TimeLogger( TEXT("SmartMerge"), [&]( FText Text) { this->LogInfo( Text ); });
-#endif
-
-	SmartMerge( World, Actors );
 }
 
-void UDataprepSmartMergeOperation::SmartMerge(UWorld* World, const TArray<AActor*>& Actors)
+void UDataprepMergeActorsOperation::PrepareStaticMeshes(const TSet<UStaticMesh*>& StaticMeshes, IMeshBuilderModule& MeshBuilderModule)
 {
-	TRACE_CPUPROFILER_EVENT_SCOPE(UDataprepSmartMergeOperation::SmartMerge);
-
-	TSharedPtr<FDataprepWorkReporter> Task = CreateTask( NSLOCTEXT( "SmartMergeOperation", "RunMerge", "Executing operation ..." ), 100.f, 1.0f );
-
-	Task->ReportNextStep( NSLOCTEXT( "SmartMergeOperation", "FindingActors", "Analyzing scene \"{0}\" ..."), 10.f );
-
-	// Group actor to merge by number of valid components
-	typedef TPair< AActor*, TArray<UPrimitiveComponent*> > FMergeableActor;
-	TMap< int32, TArray< FMergeableActor > > MergeableActorsMap;
-	for(AActor* Actor : Actors)
-	{
-		TArray<AActor*> ChildActors;
-		Actor->GetAttachedActors( ChildActors );
-
-		TArray< UPrimitiveComponent* > MeshComponents;
-		for(AActor* ChildActor : ChildActors)
-		{
-			TArray<UPrimitiveComponent*> Components;
-			ChildActor->GetComponents<UPrimitiveComponent>( Components );
-
-			for(UPrimitiveComponent* Component : Components)
-			{
-				// #ueent_todo: Support all primitive components
-				if(UStaticMeshComponent* StaticMeshComponent = Cast<UStaticMeshComponent>(Component))
-				{
-					if(StaticMeshComponent->GetStaticMesh() && StaticMeshComponent->GetStaticMesh()->GetNumSourceModels())
-					{
-						MeshComponents.Add( StaticMeshComponent );
-					}
-				}
-			}
-		}
-
-		TArray< FMergeableActor >& MergeableActors = MergeableActorsMap.FindOrAdd( MeshComponents.Num() );
-
-		MergeableActors.Emplace( Actor, MeshComponents );
-	}
-
-	Task->ReportNextStep( NSLOCTEXT( "SmartMergeOperation", "GroupingActors", "Grouping actors ..."), 10.f );
-
-	TArray< TArray< FMergeableActor > > MergeableActorsSet;
-	for(auto& MergeableActorsEntry : MergeableActorsMap)
-	{
-		TArray< FMergeableActor >& GroupedMergeableActors = MergeableActorsEntry.Value;
-		int32 ProcessedCount = 0;
-		TArray<bool> ProcessedEntries;
-		ProcessedEntries.AddZeroed( GroupedMergeableActors.Num() );
-
-		while(ProcessedCount < GroupedMergeableActors.Num())
-		{
-			for(int32 Index = 0; Index < ProcessedEntries.Num(); ++Index)
-			{
-				if(ProcessedEntries[Index] == false)
-				{
-					ProcessedEntries[Index] = true;
-					DatasmithEditingOperationsUtils::FMergingData ReferenceMergingData( GroupedMergeableActors[Index].Value );
-					ProcessedCount++;
-
-					TArray< FMergeableActor > MergeableActors;
-					MergeableActors.Add( GroupedMergeableActors[Index] );
-
-					for(int32 SubIndex = 0; SubIndex < ProcessedEntries.Num(); ++SubIndex)
-					{
-						if(ProcessedEntries[SubIndex] == false)
-						{
-							if( ReferenceMergingData.Equals( DatasmithEditingOperationsUtils::FMergingData( GroupedMergeableActors[SubIndex].Value ) ) )
-							{
-								MergeableActors.Add( GroupedMergeableActors[SubIndex] );
-								ProcessedEntries[SubIndex] = true;
-								ProcessedCount++;
-							}
-						}
-					}
-
-					MergeableActorsSet.Add( MergeableActors );
-				}
-			}
-		}
-	}
-
-	TArray< UObject* > ObjectsToDelete;
-
-	// Build render data for static meshes to be merged
-	IMeshBuilderModule& MeshBuilderModule = FModuleManager::LoadModuleChecked< IMeshBuilderModule >( TEXT("MeshBuilder") );
-	TSet<UStaticMesh*> StaticMeshes;
-
-	Task->ReportNextStep( NSLOCTEXT( "SmartMergeOperation", "BuildingMeshes", "Analyzing meshes ..."), 20.f );
-
-	for( auto& MergeableActors : MergeableActorsSet )
-	{
-		TArray<UPrimitiveComponent*>& PrimitiveComponents = MergeableActors[0].Value;
-
-		if( PrimitiveComponents.Num() > 1 )
-		{
-			for(UPrimitiveComponent* PrimitiveComponent : PrimitiveComponents)
-			{
-				if( UStaticMeshComponent* StaticMeshComponent = Cast< UStaticMeshComponent >( PrimitiveComponent ) )
-				{
-					if( StaticMeshComponent->GetStaticMesh()->RenderData == nullptr )
-					{
-						StaticMeshes.Add( StaticMeshComponent->GetStaticMesh() );
-					}
-					ObjectsToDelete.Add( StaticMeshComponent->GetStaticMesh() );
-				}
-			}
-		}
-	}
-
-	PrepareStaticMeshes( StaticMeshes, MeshBuilderModule );
+	TRACE_CPUPROFILER_EVENT_SCOPE(UDataprepMergeActorsOperation::PrepareStaticMeshes);
 
-	Task->ReportNextStep( NSLOCTEXT( "SmartMergeOperation", "MergingActors", "Merging actors ..."), 60.f );
+	if( StaticMeshes.Num() > 1 )
 	{
-		TSharedPtr<FDataprepWorkReporter> SubTask = CreateTask( NSLOCTEXT( "SmartMergeOperation", "MergingActors", "Merging actors ..."), MergeableActorsSet.Num(), 1.0f );
-
-		for( auto& MergeableActors : MergeableActorsSet )
-		{
-			SubTask->ReportNextStep( NSLOCTEXT( "SmartMergeOperation", "MergingActor", "Merging actor ...") );
-
-			TArray<UPrimitiveComponent*>& PrimitiveComponents = MergeableActors[0].Value;
-			if( PrimitiveComponents.Num() > 1 )
-			{
-				MergeStaticMeshActors( World, PrimitiveComponents, TEXT("SmartMerge"), false );
-
-				if(MergedMesh == nullptr)
-				{
-					continue;
-				}
-
-				// IMeshUtilities::MergeComponentsToStaticMesh seems to bake the rotation (maybe scaling) of the parent actor in the merged mesh.
-				// To compensate for this, the inverse rotation and the inverse scaling are computed and will be applied to the world transform of
-				// the mesh actors which will be created
-				// #ueent_todo: Investigate the reason why this must be done
-				FTransform BaseComponentToWorld = MergeableActors[0].Key->GetRootComponent()->GetComponentToWorld();
-				FQuat InvBaseRotation = BaseComponentToWorld.GetRotation().Inverse();
-				FVector InvBaseScale = FVector::OneVector / BaseComponentToWorld.GetScale3D();
-
-				for( FMergeableActor& MergeableActor : MergeableActors )
-				{
-					AActor* Actor = MergeableActor.Key;
-
-					FActorSpawnParameters Params;
-					Params.OverrideLevel = Actor->GetLevel();
-
-					// Place new mesh in the world
-					MergedActor = Cast<AStaticMeshActor>( CreateActor( AStaticMeshActor::StaticClass(), TEXT("SmartMergeActor") ) );
-					if(MergedActor != nullptr)
-					{
-						USceneComponent* ParentComponent = Actor->GetRootComponent()->GetAttachParent();
-						FString ActorName = Actor->GetName();
-						FName NewName = MakeUniqueObjectName( Actor->GetOuter(), Actor->GetClass() );
-						FDataprepCoreUtils::RenameObject( Actor, *NewName.ToString() );
-
-						MergedActor->GetStaticMeshComponent()->SetStaticMesh( MergedMesh );
-						FDataprepCoreUtils::RenameObject( MergedActor, *ActorName );
-						MergedActor->SetActorLabel( Actor->GetActorLabel() );
-
-						FTransform ComponentToWorld = Actor->GetRootComponent()->GetComponentToWorld();
-						ComponentToWorld.SetTranslation( ComponentToWorld.GetTranslation() /*+ MergedActorLocation*/ );
-						ComponentToWorld.SetRotation( ComponentToWorld.GetRotation() * InvBaseRotation );
-						ComponentToWorld.SetScale3D( ComponentToWorld.GetScale3D() * InvBaseScale );
-						MergedActor->GetRootComponent()->SetComponentToWorld( ComponentToWorld );
-
-						World->UpdateCullDistanceVolumes( MergedActor, MergedActor->GetStaticMeshComponent() );
-
-						if(ParentComponent)
-						{
-							MergedActor->GetRootComponent()->AttachToComponent( ParentComponent, FAttachmentTransformRules::KeepWorldTransform );
-						}
-
-						TArray<AActor*> ActorsToDelete;
-						Actor->GetAttachedActors( ActorsToDelete );
-						ActorsToDelete.Add( Actor );
-
-						ObjectsToDelete.Append( ActorsToDelete );
-					}
-					else
-					{
-						UE_LOG(LogDataprep, Error, TEXT("MergeStaticMeshActors failed. Internal error while creating the merged actor."));
-					}
-
-					TArray<AActor*> ActorsToDelete;
-					Actor->GetAttachedActors( ActorsToDelete );
-					ActorsToDelete.Add( Actor );
-
-					ObjectsToDelete.Append( ActorsToDelete );
-				}
-			}
-		}
+		TArray<UStaticMesh*> StaticMeshesToBuild = StaticMeshes.Array();
+		ParallelFor( StaticMeshesToBuild.Num(), [&]( int32 Index ) {
+			DataprepOperationsLibraryUtil::BuildRenderData( StaticMeshesToBuild[Index], MeshBuilderModule );
+		});
 	}
-
-	// Release render data of built static meshes
-	for(UStaticMesh* StaticMesh : StaticMeshes)
+	else
 	{
-		if(StaticMesh)
+		for(UStaticMesh* StaticMesh : StaticMeshes)
 		{
-			StaticMesh->RenderData.Reset();
+			DataprepOperationsLibraryUtil::BuildRenderData( StaticMesh, MeshBuilderModule );
 		}
 	}
-
-	DeleteObjects(ObjectsToDelete);
 }
 
-void UDataprepCleanWorldOperation::OnExecution_Implementation(const FDataprepContext& InContext)
+void UDataprepDeleteUnusedAssetsOperation::OnExecution_Implementation(const FDataprepContext& InContext)
 {
 	UWorld* World = nullptr;
 
@@ -577,12 +383,11 @@ void UDataprepCleanWorldOperation::OnExecution_Implementation(const FDataprepCon
 			}
 		}
 
-		for(const FMaterialTextureInfo& TextureInfo : MaterialInterface->GetTextureStreamingData())
+		TArray<UTexture*> Textures;
+		MaterialInterface->GetUsedTextures(Textures, EMaterialQualityLevel::Num, true, ERHIFeatureLevel::Num, true);
+		for (UTexture* Texture : Textures)
 		{
-			if(UObject* Texture = TextureInfo.TextureReference.ResolveObject())
-			{
-				UsedAssets.Add(Texture);
-			}
+			UsedAssets.Add(Texture);
 		}
 	};
 
@@ -640,56 +445,89 @@ void UDataprepCleanWorldOperation::OnExecution_Implementation(const FDataprepCon
 	TArray<UObject*> ObjectsToDelete;
 	ObjectsToDelete.Reserve(InContext.Objects.Num());
 
-	if(World != nullptr)
+	for(UObject* Object : InContext.Objects)
 	{
-		// Get root actors
-		TArray<AActor*> RootActors;
-		DatasmithEditingOperationsUtils::GetRootActors(World, RootActors);
+		if(FDataprepCoreUtils::IsAsset(Object) && !UsedAssets.Contains(Object))
+		{
+			ObjectsToDelete.Add(Object);
+		}
+	}
 
-		TArray<AActor*> ActorsToCollapse;
-		DatasmithEditingOperationsUtils::FindActorsToCollapseOrDelete(RootActors, ActorsToCollapse, ObjectsToDelete);
+	DeleteObjects( ObjectsToDelete );
+}
+
+void UDataprepCompactSceneGraphOperation::OnExecution_Implementation(const FDataprepContext& InContext)
+{
+#ifdef LOG_TIME
+	DataprepEditingOperationTime::FTimeLogger TimeLogger(TEXT("CompactSceneGraph"), [&](FText Text) { this->LogInfo(Text); });
+#endif
 
-		for(AActor* Actor : ActorsToCollapse)
+	TMap<AActor*, bool> VisibilityMap;
+	for (UObject* Object : InContext.Objects)
+	{
+		if (!ensure(Object) || Object->IsPendingKill())
 		{
-			USceneComponent* RootComponent = Actor->GetRootComponent();
+			continue;
+		}
 
-			TArray< USceneComponent* > AttachChildren = RootComponent->GetAttachChildren(); // Make a copy because the array in RootComponent will get modified during the process
-			USceneComponent* AttachParent = RootComponent->GetAttachParent();
+		if (AActor* Actor = Cast<AActor>(Object))
+		{
+			IsActorVisible(Actor, VisibilityMap);
+		}
+	}
 
-			for ( USceneComponent* ChildComponent : AttachChildren )
-			{
-				if(ChildComponent)
-				{
-					// skip component with invalid or condemned owner
-					AActor* Owner = ChildComponent->GetOwner();
-					if ( Owner == nullptr || Owner == Actor || Owner->IsPendingKill() || InContext.Objects.Contains(Owner) /* Slow!!! */)
-					{
-						continue;
-					}
+	TArray<UObject*> ObjectsToDelete;
+	ObjectsToDelete.Reserve(InContext.Objects.Num());
 
-					ChildComponent->AttachToComponent( AttachParent, FAttachmentTransformRules::KeepWorldTransform );
-				}
-			}
+	for (const TPair<AActor*, bool>& ActorVisiblity : VisibilityMap)
+	{
+		if (!ActorVisiblity.Value)
+		{
+			ObjectsToDelete.Add(ActorVisiblity.Key);
+		}
+	}
 
-			// Remove actor from world and add to array for deletion
-			{
-				World->RemoveActor(Actor, true);
-				FDataprepCoreUtils::MoveToTransientPackage( Actor );
+	DeleteObjects(ObjectsToDelete);
+}
 
-				ObjectsToDelete.Add( Actor );
-			}
+bool UDataprepCompactSceneGraphOperation::IsActorVisible(AActor* Actor, TMap<AActor*, bool>& VisibilityMap)
+{
+	if (!Actor)
+	{
+		return false;
+	}
+
+	// For scene compaction, actor visibility is defined as the actor having a MeshComponent (PrimitiveComponent could also be used)
+	// or an attached child that is visible
+	bool* bIsVisible = VisibilityMap.Find(Actor);
+	if (bIsVisible)
+	{
+		return *bIsVisible;
+	}
+
+	TArray<UActorComponent*> Components = Actor->GetComponents().Array();
+	for (UActorComponent* Component : Components)
+	{
+		if (UMeshComponent* MeshComponent = Cast<UMeshComponent>(Component))
+		{
+			VisibilityMap.Add(Actor, true);
+			return true;
 		}
 	}
 
-	for(UObject* Object : InContext.Objects)
+	TArray<AActor*> AttachedActors;
+	Actor->GetAttachedActors(AttachedActors);
+	for (AActor* AttachedActor : AttachedActors)
 	{
-		if(FDataprepCoreUtils::IsAsset(Object) && !UsedAssets.Contains(Object))
+		if (IsActorVisible(AttachedActor, VisibilityMap))
 		{
-			ObjectsToDelete.Add(Object);
+			VisibilityMap.Add(Actor, true);
+			return true;
 		}
 	}
 
-	DeleteObjects( ObjectsToDelete );
+	VisibilityMap.Add(Actor, false);
+	return false;
 }
 
 namespace DatasmithEditingOperationsUtils
diff --git a/Engine/Plugins/Experimental/DataPrepEditor/Source/DataPrepLibraries/Private/DataprepEditingOperations.h b/Engine/Plugins/Experimental/DataPrepEditor/Source/DataPrepLibraries/Private/DataprepEditingOperations.h
--- a/Engine/Plugins/Experimental/DataPrepEditor/Source/DataPrepLibraries/Private/DataprepEditingOperations.h
+++ b/Engine/Plugins/Experimental/DataPrepEditor/Source/DataPrepLibraries/Private/DataprepEditingOperations.h
@@ -17,8 +17,8 @@ class UStaticMesh;
 class UStaticMeshComponent;
 class UWorld;
 
-UCLASS(Experimental, Category = ObjectOperation, Meta = (DisplayName="Remove Objects", ToolTip = "Remove any asset or actor to process") )
-class UDataprepRemoveObjectsOperation : public UDataprepEditingOperation
+UCLASS(Experimental, Category = ObjectOperation, Meta = (DisplayName="Delete Objects", ToolTip = "Delete any asset or actor to process") )
+class UDataprepDeleteObjectsOperation : public UDataprepEditingOperation
 {
 	GENERATED_BODY()
 
@@ -41,8 +41,9 @@ class UDataprepMergeActorsOperation : public UDataprepEditingOperation
 
 public:
 	UDataprepMergeActorsOperation()
-		: bDeleteMergedActors(true)
-		, bDeleteMergedMeshes(true)
+		: bDeleteMergedActors_DEPRECATED(true)
+		, bDeleteMergedMeshes_DEPRECATED(true)
+		, bPivotPointAtZero(false)
 		, MergedMesh(nullptr)
 		, MergedActor(nullptr)
 	{
@@ -60,24 +61,36 @@ public:
 	FString NewActorLabel;
 
 	/** Settings to use for the merge operation */
-	UPROPERTY(EditAnywhere, Category = MergeSettings)
-	bool bDeleteMergedActors;
+	UPROPERTY()
+	bool bDeleteMergedActors_DEPRECATED;
 
 	/** Settings to use for the merge operation */
-	UPROPERTY(EditAnywhere, Category = MergeSettings)
-	bool bDeleteMergedMeshes;
+	UPROPERTY()
+	bool bDeleteMergedMeshes_DEPRECATED;
 
-	/** Settings to use for the merge operation */
-	UPROPERTY(EditAnywhere, Category = MergeSettings)
-	FMeshMergingSettings MergeSettings;
+	UPROPERTY()
+	FMeshMergingSettings MergeSettings_DEPRECATED;
+
+	/** Whether merged mesh should have pivot at world origin, or at first merged component otherwise */
+	UPROPERTY(EditAnywhere, BlueprintReadWrite, Category = MeshSettings)
+	bool bPivotPointAtZero;
 
 protected:
+	class FStaticMeshBuilder
+	{
+	public:
+		FStaticMeshBuilder(const TSet<UStaticMesh*>& InStaticMeshes);
+		~FStaticMeshBuilder();
+	private:
+		TSet<UStaticMesh*> StaticMeshes;
+	};
+
 	virtual void OnExecution_Implementation(const FDataprepContext& InContext) override;
 	//~ End UDataprepOperation Interface
 
 	bool MergeStaticMeshActors(UWorld* World, const TArray<UPrimitiveComponent*>& ComponentsToMerge, const FString& RootName, bool bCreateActor = true);
 
-	void PrepareStaticMeshes( TSet<UStaticMesh*> StaticMeshes, IMeshBuilderModule& MeshBuilderModule );
+	static void PrepareStaticMeshes( const TSet<UStaticMesh*>& StaticMeshes, IMeshBuilderModule& MeshBuilderModule );
 
 protected:
 	FVector MergedMeshWorldLocation;
@@ -85,33 +98,38 @@ protected:
 	AStaticMeshActor* MergedActor;
 };
 
-UCLASS(Experimental, Category = ObjectOperation, Meta = (DisplayName="Smart Merge - PROTO", ToolTip = "Collapse all actors solely holding more than one static mesh actor") )
-class UDataprepSmartMergeOperation : public UDataprepMergeActorsOperation
+UCLASS(Experimental, Category = ObjectOperation, Meta = (DisplayName="Delete Unused Assets", ToolTip = "Delete assets that are not referenced by any objects") )
+class UDataprepDeleteUnusedAssetsOperation : public UDataprepEditingOperation
 {
 	GENERATED_BODY()
 
 	//~ Begin UDataprepOperation Interface
+public:
+	virtual FText GetCategory_Implementation() const override
+	{
+		return FDataprepOperationCategories::ObjectOperation;
+	}
+
 protected:
 	virtual void OnExecution_Implementation(const FDataprepContext& InContext) override;
 	//~ End UDataprepOperation Interface
-
-private:
-	void SmartMerge(UWorld* World, const TArray<AActor*>& Actors);
 };
 
-UCLASS(Experimental, Category = ObjectOperation, Meta = (DisplayName="Clean World - PROTO", ToolTip = "Remove unused assets, collapse actors with only one child") )
-class UDataprepCleanWorldOperation : public UDataprepEditingOperation
+UCLASS(Experimental, Category = ActorOperation, Meta = (DisplayName="Compact Scene Graph", ToolTip = "Delete actors that do not have visuals, but keep those needed to preserve hierarchy") )
+class UDataprepCompactSceneGraphOperation : public UDataprepEditingOperation
 {
 	GENERATED_BODY()
 
-		//~ Begin UDataprepOperation Interface
+	//~ Begin UDataprepOperation Interface
 public:
 	virtual FText GetCategory_Implementation() const override
 	{
-		return FDataprepOperationCategories::ObjectOperation;
+		return FDataprepOperationCategories::ActorOperation;
 	}
 
 protected:
 	virtual void OnExecution_Implementation(const FDataprepContext& InContext) override;
 	//~ End UDataprepOperation Interface
+
+	bool IsActorVisible(AActor*, TMap<AActor*, bool>& VisibilityMap);
 };
