Merging CL #9421202 from //UE4/Dev-Enterprise

DataPrep operations clean-up.

diff --git a/Engine/Plugins/Experimental/DataPrepEditor/Source/DataPrepLibraries/Private/DataprepOperations.cpp b/Engine/Plugins/Experimental/DataPrepEditor/Source/DataPrepLibraries/Private/DataprepOperations.cpp
--- a/Engine/Plugins/Experimental/DataPrepEditor/Source/DataPrepLibraries/Private/DataprepOperations.cpp
+++ b/Engine/Plugins/Experimental/DataPrepEditor/Source/DataPrepLibraries/Private/DataprepOperations.cpp
@@ -145,7 +145,7 @@ void UDataprepSetMobilityOperation::OnExecution_Implementation(const FDataprepCo
 
 void UDataprepSetMaterialOperation::OnExecution_Implementation(const FDataprepContext& InContext)
 {
-	if(MaterialSubstitute == nullptr)
+	if(Material == nullptr)
 	{
 		FText OutReason = FText( LOCTEXT( "DatasmithMeshOperations_SetMaterial", "No material specified. Aborting operation..." ) );
 		LogInfo( OutReason );
@@ -159,7 +159,7 @@ void UDataprepSetMaterialOperation::OnExecution_Implementation(const FDataprepCo
 #endif
 
 	// Execute operation
-	UDataprepOperationsLibrary::SetMaterial( InContext.Objects, MaterialSubstitute );
+	UDataprepOperationsLibrary::SetMaterial( InContext.Objects, Material );
 }
 
 void UDataprepSubstituteMaterialOperation::OnExecution_Implementation(const FDataprepContext& InContext)
@@ -281,7 +281,7 @@ void FDataprepSetLOGGroupDetails::CustomizeDetails(IDetailLayoutBuilder & Detail
 
 void UDataprepSetMeshOperation::OnExecution_Implementation(const FDataprepContext& InContext)
 {
-	if(MeshSubstitute == nullptr)
+	if(StaticMesh == nullptr)
 	{
 		FText OutReason = FText( LOCTEXT( "DatasmithMeshOperations_SetMesh", "No mesh specified. Aborting operation..." ) );
 		LogInfo( OutReason );
@@ -295,45 +295,7 @@ void UDataprepSetMeshOperation::OnExecution_Implementation(const FDataprepContex
 #endif
 
 	// Execute operation
-	UDataprepOperationsLibrary::SetMesh( InContext.Objects, MeshSubstitute );
-}
-
-void UDataprepSubstituteMeshOperation::OnExecution_Implementation(const FDataprepContext& InContext)
-{
-	if(MeshSubstitute == nullptr)
-	{
-		FText OutReason = FText( LOCTEXT( "DatasmithMeshOperations_SubstituteMesh", "No mesh specified. Aborting operation..." ) );
-		LogInfo( OutReason );
-		// #ueent_todo: Remove call to UE_LOG when DataprepLogger is operational
-		UE_LOG(LogDataprep, Log, TEXT("UDataprepSubstituteMeshOperation: %s"), *OutReason.ToString());
-		return;
-	}
-
-#ifdef LOG_TIME
-	DataprepOperationTime::FTimeLogger TimeLogger( TEXT("SubstituteMesh"), [&]( FText Text) { this->LogInfo( Text ); });
-#endif
-
-	// Execute operation
-	UDataprepOperationsLibrary::SubstituteMesh( InContext.Objects, MeshSearch, StringMatch, MeshSubstitute );
-}
-
-void UDataprepSubstituteMeshByTableOperation::OnExecution_Implementation(const FDataprepContext& InContext)
-{
-	if(MeshDataTable == nullptr)
-	{
-		FText OutReason = FText( LOCTEXT( "DatasmithMeshOperations_SubstituteMeshByTable", "No data table specified. Aborting operation..." ) );
-		LogInfo( OutReason );
-		// #ueent_todo: Remove call to UE_LOG when DataprepLogger is operational
-		UE_LOG(LogDataprep, Log, TEXT("UDataprepSubstituteMeshByTableOperation: %s"), *OutReason.ToString());
-		return;
-	}
-
-#ifdef LOG_TIME
-	DataprepOperationTime::FTimeLogger TimeLogger( TEXT("SubstituteMeshesByTable"), [&]( FText Text) { this->LogInfo( Text ); });
-#endif
-
-	// Execute operation
-	UDataprepOperationsLibrary::SubstituteMeshesByTable( InContext.Objects, MeshDataTable );
+	UDataprepOperationsLibrary::SetMesh( InContext.Objects, StaticMesh );
 }
 
 #undef LOCTEXT_NAMESPACE
diff --git a/Engine/Plugins/Experimental/DataPrepEditor/Source/DataPrepLibraries/Private/DataprepOperations.h b/Engine/Plugins/Experimental/DataPrepEditor/Source/DataPrepLibraries/Private/DataprepOperations.h
--- a/Engine/Plugins/Experimental/DataPrepEditor/Source/DataPrepLibraries/Private/DataprepOperations.h
+++ b/Engine/Plugins/Experimental/DataPrepEditor/Source/DataPrepLibraries/Private/DataprepOperations.h
@@ -224,14 +224,14 @@ class UDataprepSetMaterialOperation : public UDataprepOperation
 	GENERATED_BODY()
 
 	UDataprepSetMaterialOperation()
-		: MaterialSubstitute(nullptr)
+		: Material(nullptr)
 	{
 	}
 
 public:
 	// Material to use as a substitute
 	UPROPERTY(EditAnywhere, BlueprintReadWrite, Category = MeshOperation, meta = (ToolTip = "Material to use as a substitute"))
-	UMaterialInterface* MaterialSubstitute;
+	UMaterialInterface* Material;
 
 	//~ Begin UDataprepOperation Interface
 public:
@@ -341,78 +341,14 @@ class UDataprepSetMeshOperation : public UDataprepOperation
 	GENERATED_BODY()
 
 	UDataprepSetMeshOperation()
-	: MeshSubstitute(nullptr)
+	: StaticMesh(nullptr)
 	{
 	}
 
 public:
 	// Mesh to use as a substitute
 	UPROPERTY(EditAnywhere, BlueprintReadWrite, Category = ActorOperation, meta = (ToolTip = "Mesh to use as a substitute"))
-	UStaticMesh* MeshSubstitute;
-
-	//~ Begin UDataprepOperation Interface
-public:
-	virtual FText GetCategory_Implementation() const override
-	{
-		return FDataprepOperationCategories::ActorOperation;
-	}
-
-protected:
-	virtual void OnExecution_Implementation(const FDataprepContext& InContext) override;
-	//~ End UDataprepOperation Interface
-};
-
-UCLASS(Experimental, Category = ActorOperation, Meta = (DisplayName="Substitute Mesh", ToolTip = "On each actor to process, replace the mesh matching the criteria with the specified one") )
-class UDataprepSubstituteMeshOperation : public UDataprepOperation
-{
-	GENERATED_BODY()
-
-	UDataprepSubstituteMeshOperation()
-	: MeshSearch(TEXT("*"))
-	, StringMatch(EEditorScriptingStringMatchType::MatchesWildcard)
-	, MeshSubstitute(nullptr)
-	{
-	}
-
-public:
-	// Name of the mesh(s) to search for. Wildcard is supported
-	UPROPERTY(EditAnywhere, BlueprintReadWrite, Category = ActorOperation, meta = (ToolTip = "Name of the mesh(s) to search for. Wildcard is supported"))
-	FString MeshSearch;
-
-	// Type of matching to perform with MeshSearch string
-	UPROPERTY(EditAnywhere, BlueprintReadWrite, Category = ActorOperation, meta = (ToolTip = "Type of matching to perform with MeshSearch string"))
-	EEditorScriptingStringMatchType StringMatch;
-
-	// Mesh to use as a substitute
-	UPROPERTY(EditAnywhere, BlueprintReadWrite, Category = ActorOperation, meta = (ToolTip = "Mesh to use as a substitute"))
-	UStaticMesh* MeshSubstitute;
-
-	//~ Begin UDataprepOperation Interface
-public:
-	virtual FText GetCategory_Implementation() const override
-	{
-		return FDataprepOperationCategories::ActorOperation;
-	}
-
-protected:
-	virtual void OnExecution_Implementation(const FDataprepContext& InContext) override;
-	//~ End UDataprepOperation Interface
-};
-
-UCLASS(Experimental, Category = ActorOperation, Meta = (DisplayName="Substitute Mesh By Table", ToolTip = "On each actor to process, replace the mesh found in the first column of the table with the one from the second column in the same row") )
-class UDataprepSubstituteMeshByTableOperation : public UDataprepOperation
-{
-	GENERATED_BODY()
-
-	UDataprepSubstituteMeshByTableOperation()
-	: MeshDataTable(nullptr)
-	{
-	}
-
-public:
-	// Data table to use for the substitution
-	UPROPERTY(EditAnywhere, BlueprintReadWrite, Category = ActorOperation, meta = (ToolTip = "Data table to use for the substitution"))
-	UDataTable* MeshDataTable;
+	UStaticMesh* StaticMesh;
 
 	//~ Begin UDataprepOperation Interface
 public:
