Reworked generation of rendering data for static meshes and compilation of materials: Now done after the run of the producers and during the execution of the pipeline
Modified the logic in the viewport to take in account the new logic with regard to the generation of render data
Used UStaticMesh::BatchBuild to generate missing render data for operations
Added logic to properly track assets or actors which have been modifed

diff --git a/Engine/Plugins/Experimental/DataPrepEditor/Source/DataPrepCore/Private/AssetTypeActions_DataprepAssetInterface.cpp b/Engine/Plugins/Experimental/DataPrepEditor/Source/DataPrepCore/Private/AssetTypeActions_DataprepAssetInterface.cpp
--- a/Engine/Plugins/Experimental/DataPrepEditor/Source/DataPrepCore/Private/AssetTypeActions_DataprepAssetInterface.cpp
+++ b/Engine/Plugins/Experimental/DataPrepEditor/Source/DataPrepCore/Private/AssetTypeActions_DataprepAssetInterface.cpp
@@ -72,7 +72,7 @@ void FAssetTypeActions_DataprepAssetInterface::CreateInstance(TArray<TWeakObject
 				Factory->InitialParent = DataprepAssetInterface;
 
 				IAssetTools& AssetTools = FModuleManager::LoadModuleChecked<FAssetToolsModule>("AssetTools").Get();
-				UObject* NewAsset = AssetTools.CreateAsset(Name, FPackageName::GetLongPackagePath(PackageName), UDataprepAssetInterface::StaticClass(), Factory);
+				UObject* NewAsset = AssetTools.CreateAsset(Name, FPackageName::GetLongPackagePath(PackageName), UDataprepAssetInstance::StaticClass(), Factory);
 
 				if(NewAsset)
 				{
diff --git a/Engine/Plugins/Experimental/DataPrepEditor/Source/DataPrepCore/Private/DataPrepAssetProducers.cpp b/Engine/Plugins/Experimental/DataPrepEditor/Source/DataPrepCore/Private/DataPrepAssetProducers.cpp
--- a/Engine/Plugins/Experimental/DataPrepEditor/Source/DataPrepCore/Private/DataPrepAssetProducers.cpp
+++ b/Engine/Plugins/Experimental/DataPrepEditor/Source/DataPrepCore/Private/DataPrepAssetProducers.cpp
@@ -94,6 +94,8 @@ TArray< TWeakObjectPtr< UObject > > UDataprepAssetProducers::Produce( const FDat
 		}
 	}
 
+	FDataprepCoreUtils::BuildAssets( OutAssets, InContext.ProgressReporterPtr );
+
 	return OutAssets;
 }
 
@@ -113,8 +115,16 @@ UDataprepContentProducer* UDataprepAssetProducers::AddProducer(UClass* ProducerC
 		int32 ProducerNextIndex = AssetProducers.Num();
 		AssetProducers.Emplace( Producer, true );
 
+		bool bChangeAll = false;
+		ValidateProducerChanges( ProducerNextIndex, bChangeAll );
+
 		OnChanged.Broadcast( FDataprepAssetChangeType::ProducerAdded, ProducerNextIndex );
 
+		if(bChangeAll)
+		{
+			OnChanged.Broadcast( FDataprepAssetChangeType::ProducerModified, INDEX_NONE );
+		}
+
 		return Producer;
 	}
 
@@ -138,8 +148,16 @@ UDataprepContentProducer* UDataprepAssetProducers::CopyProducer(const UDataprepC
 		int32 ProducerNextIndex = AssetProducers.Num();
 		AssetProducers.Emplace( Producer, true );
 
+		bool bChangeAll = false;
+		ValidateProducerChanges( ProducerNextIndex, bChangeAll );
+
 		OnChanged.Broadcast( FDataprepAssetChangeType::ProducerAdded, ProducerNextIndex );
 
+		if(bChangeAll)
+		{
+			OnChanged.Broadcast( FDataprepAssetChangeType::ProducerModified, INDEX_NONE );
+		}
+
 		return Producer;
 	}
 
diff --git a/Engine/Plugins/Experimental/DataPrepEditor/Source/DataPrepCore/Private/DataPrepOperation.cpp b/Engine/Plugins/Experimental/DataPrepEditor/Source/DataPrepCore/Private/DataPrepOperation.cpp
--- a/Engine/Plugins/Experimental/DataPrepEditor/Source/DataPrepCore/Private/DataPrepOperation.cpp
+++ b/Engine/Plugins/Experimental/DataPrepEditor/Source/DataPrepCore/Private/DataPrepOperation.cpp
@@ -83,6 +83,12 @@ TSharedPtr<FDataprepWorkReporter> UDataprepOperation::CreateTask(const FText & I
 	return TSharedPtr<FDataprepWorkReporter>();
 }
 
+bool UDataprepOperation::IsCancelled()
+{
+	return OperationContext && OperationContext->DataprepProgressReporter ? OperationContext->DataprepProgressReporter->IsWorkCancelled() : false;
+}
+
+
 void UDataprepOperation::ExecuteOperation(TSharedRef<FDataprepOperationContext>&  InOperationContext)
 {
 	OperationContext = InOperationContext;
@@ -116,11 +122,31 @@ FText UDataprepOperation::GetAdditionalKeyword_Implementation() const
 	return FText();
 }
 
-UObject* UDataprepEditingOperation::AddAsset( const UObject* Asset, UClass* AssetClass, const FString& AssetName )
+void UDataprepOperation::AssetsModified( TArray<UObject*> Assets )
+{
+	if ( OperationContext && OperationContext->AssetsModifiedDelegate.IsBound() )
+	{
+		OperationContext->AssetsModifiedDelegate.Execute( MoveTemp( Assets ) );
+	}
+}
+
+UObject* UDataprepEditingOperation::AddAsset( const UObject* Asset, const FString& AssetName )
 {
 	if ( OperationContext && OperationContext->AddAssetDelegate.IsBound() )
 	{
-		return OperationContext->AddAssetDelegate.Execute( Asset, AssetClass, AssetName.Len() > 0 ? *AssetName : nullptr );
+		return OperationContext->AddAssetDelegate.Execute( Asset, AssetName.Len() > 0 ? *AssetName : nullptr );
+	}
+
+	// #ueent_todo: Report error
+
+	return nullptr;
+}
+
+UObject* UDataprepEditingOperation::CreateAsset( UClass* AssetClass, const FString& AssetName )
+{
+	if ( OperationContext && OperationContext->CreateAssetDelegate.IsBound() )
+	{
+		return OperationContext->CreateAssetDelegate.Execute( AssetClass, AssetName.Len() > 0 ? *AssetName : nullptr );
 	}
 
 	// #ueent_todo: Report error
@@ -171,7 +197,7 @@ void UDataprepEditingOperation::DeleteObject( UObject* Object )
 	{
 		TArray<UObject*> Objects;
 		Objects.Add( Object );
-		OperationContext->DeleteObjectsDelegate.Execute( Objects );
+		OperationContext->DeleteObjectsDelegate.Execute( MoveTemp( Objects ) );
 		return;
 	}
 
@@ -182,8 +208,7 @@ void UDataprepEditingOperation::DeleteObjects( TArray<UObject*> Objects )
 {
 	if ( OperationContext && OperationContext->DeleteObjectsDelegate.IsBound() )
 	{
-		OperationContext->DeleteObjectsDelegate.Execute( Objects );
-		return;
+		OperationContext->DeleteObjectsDelegate.Execute( MoveTemp( Objects ) );
 	}
 
 	// #ueent_todo: Report error
diff --git a/Engine/Plugins/Experimental/DataPrepEditor/Source/DataPrepCore/Private/DataprepActionAsset.cpp b/Engine/Plugins/Experimental/DataPrepEditor/Source/DataPrepCore/Private/DataprepActionAsset.cpp
--- a/Engine/Plugins/Experimental/DataPrepEditor/Source/DataPrepCore/Private/DataprepActionAsset.cpp
+++ b/Engine/Plugins/Experimental/DataPrepEditor/Source/DataPrepCore/Private/DataprepActionAsset.cpp
@@ -22,6 +22,7 @@
 #include "GameFramework/WorldSettings.h"
 #include "LevelSequence.h"
 #include "Materials/MaterialInterface.h"
+#include "ObjectTools.h"
 #include "Templates/SubclassOf.h"
 #include "UObject/ObjectMacros.h"
 #include "UObject/UObjectHash.h"
@@ -46,9 +47,11 @@ UDataprepActionAsset::UDataprepActionAsset()
 		OperationContext->Context = MakeShareable( new FDataprepContext() );
 
 		OperationContext->AddAssetDelegate = FDataprepAddAsset::CreateUObject( this, &UDataprepActionAsset::OnAddAsset );
+		OperationContext->CreateAssetDelegate = FDataprepCreateAsset::CreateUObject( this, &UDataprepActionAsset::OnCreateAsset );
 		OperationContext->CreateActorDelegate = FDataprepCreateActor::CreateUObject( this, &UDataprepActionAsset::OnCreateActor );
 		OperationContext->RemoveObjectDelegate = FDataprepRemoveObject::CreateUObject( this, &UDataprepActionAsset::OnRemoveObject );
 		OperationContext->DeleteObjectsDelegate = FDataprepDeleteObjects::CreateUObject( this, &UDataprepActionAsset::OnDeleteObjects );
+		OperationContext->AssetsModifiedDelegate = FDataprepAssetsModified::CreateUObject( this, &UDataprepActionAsset::OnAssetsModified );
 	}
 }
 
@@ -379,16 +382,8 @@ void UDataprepActionAsset::ExecuteAction(const TSharedPtr<FDataprepActionContext
 			TSharedRef<FDataprepOperationContext> OperationContextPtr = this->OperationContext.ToSharedRef();
 			Operation->ExecuteOperation( OperationContextPtr );
 
-			// Check that only UDataprepEditingOperation are changing the content
-			if(this->bWorkingSetHasChanged)
-			{
-				ensure( Operation->IsA<UDataprepEditingOperation>() );
-			}
-
 			// Process the changes in the context if applicable
-			AssetsDiffCount = ContextPtr->Assets.Num() - AssetsDiffCount;
-			ActorsDiffCount = OperationContext->Context->Objects.Num() - ActorsDiffCount - AssetsDiffCount;
-			this->ProcessWorkingSetChanged( AssetsDiffCount != 0, ActorsDiffCount != 0 );
+			this->ProcessWorkingSetChanged();
 		}
 		else if ( UDataprepFilter* Filter = Step->Filter )
 		{
@@ -439,41 +434,64 @@ void UDataprepActionAsset::ExecuteAction(const TSharedPtr<FDataprepActionContext
 	ContextPtr.Reset();
 }
 
-UObject* UDataprepActionAsset::OnAddAsset(const UObject* Asset, UClass* AssetClass, const TCHAR* AssetName)
+UObject* UDataprepActionAsset::OnAddAsset(const UObject* Asset, const TCHAR* AssetName)
 {
 	UObject* NewAsset = nullptr;
 
-	if(ContextPtr != nullptr)
+	if(ContextPtr != nullptr && Asset != nullptr)
 	{
-		UObject* Outer = GetAssetOuterByClass( Asset ? Asset->GetClass() : AssetClass );
+		UObject* Outer = GetAssetOuterByClass( Asset->GetClass() );
 
-		// #ueent_todo: Deal with name collision
-		if(Asset == nullptr)
-		{
-			NewAsset = NewObject<UObject>( Outer, AssetClass, NAME_None, RF_Transient );
-		}
-		else
-		{
-			NewAsset = DuplicateObject< UObject >( Asset, Outer, NAME_None );
-		}
+		NewAsset = DuplicateObject< UObject >( Asset, Outer, NAME_None );
 		check( NewAsset );
 
-		if(AssetName != nullptr)
-		{
-			FName UniqueName = MakeUniqueObjectName( Outer, NewAsset->GetClass(), AssetName );
-			FDataprepCoreUtils::RenameObject( NewAsset, *UniqueName.ToString() );
-		}
+		AddAssetToContext( NewAsset, AssetName );
+	}
 
-		// Add new asset to local and global contexts
-		ContextPtr->Assets.Add( NewAsset );
-		OperationContext->Context->Objects.Add( NewAsset );
+	return NewAsset;
+}
 
-		bWorkingSetHasChanged = true;
+UObject* UDataprepActionAsset::OnCreateAsset(UClass* AssetClass, const TCHAR* AssetName)
+{
+	UObject* NewAsset = nullptr;
+
+	if(ContextPtr != nullptr && AssetClass != nullptr)
+	{
+		UObject* Outer = GetAssetOuterByClass( AssetClass );
+
+		NewAsset = NewObject<UObject>( Outer, AssetClass, NAME_None, RF_Transient );
+		check( NewAsset );
+
+		AddAssetToContext( NewAsset, AssetName );
 	}
 
 	return NewAsset;
 }
 
+void UDataprepActionAsset::AddAssetToContext( UObject* NewAsset, const TCHAR* DesiredName )
+{
+	check( NewAsset );
+
+	if(DesiredName != nullptr)
+	{
+		// Rename producer to name of file
+		FString AssetName = ObjectTools::SanitizeObjectName( DesiredName );
+		if ( !NewAsset->Rename( *AssetName, nullptr, REN_Test ) )
+		{
+			AssetName = MakeUniqueObjectName( GetOuter(), GetClass(), *AssetName ).ToString();
+		}
+
+		FDataprepCoreUtils::RenameObject( NewAsset, *AssetName );
+	}
+
+	// Add new asset to local and global contexts
+	ContextPtr->Assets.Add( NewAsset );
+	OperationContext->Context->Objects.Add( NewAsset );
+
+	AddedObjects.Add( NewAsset );
+	bWorkingSetHasChanged = true;
+}
+
 AActor* UDataprepActionAsset::OnCreateActor(UClass* ActorClass, const TCHAR* ActorName)
 {
 	if(ActorClass != nullptr && ContextPtr != nullptr)
@@ -489,6 +507,7 @@ AActor* UDataprepActionAsset::OnCreateActor(UClass* ActorClass, const TCHAR* Act
 		// Add new actor to local contexts
 		OperationContext->Context->Objects.Add( Actor );
 
+		AddedObjects.Add( Actor );
 		bWorkingSetHasChanged = true;
 
 		return Actor;
@@ -507,6 +526,20 @@ void UDataprepActionAsset::OnRemoveObject(UObject* Object, bool bLocalContext)
 	}
 }
 
+void UDataprepActionAsset::OnAssetsModified(TArray<UObject*> Assets)
+{
+	if(ContextPtr != nullptr)
+	{
+		for(UObject* Asset : Assets)
+		{
+			if(Asset != nullptr)
+			{
+				ModifiedAssets.Add( Asset );
+			}
+		}
+	}
+}
+
 void UDataprepActionAsset::OnDeleteObjects(TArray<UObject*> Objects)
 {
 	if(ContextPtr != nullptr)
@@ -516,6 +549,11 @@ void UDataprepActionAsset::OnDeleteObjects(TArray<UObject*> Objects)
 			// Mark object for deletion
 			if(Object != nullptr)
 			{
+				if( FDataprepCoreUtils::IsAsset( Object ) && ModifiedAssets.Contains( Object ) )
+				{
+					ModifiedAssets.Remove( Object );
+				}
+
 				ObjectsToDelete.Add(Object);
 				bWorkingSetHasChanged = true;
 			}
@@ -523,12 +561,19 @@ void UDataprepActionAsset::OnDeleteObjects(TArray<UObject*> Objects)
 	}
 }
 
-void UDataprepActionAsset::ProcessWorkingSetChanged( bool bAddedAssets, bool bAddedActors )
+void UDataprepActionAsset::ProcessWorkingSetChanged()
 {
-	if(bWorkingSetHasChanged && ContextPtr != nullptr)
+	if((bWorkingSetHasChanged || ModifiedAssets.Num() > 0) && ContextPtr != nullptr)
 	{
-		bool bAssetsChanged = bAddedAssets;
-		bool bWorldChanged = bAddedActors;
+		bool bAssetsChanged = ModifiedAssets.Num() > 0;
+		bool bWorldChanged = false;
+
+		for(UObject* Object : AddedObjects)
+		{
+			const AActor* Actor = Cast<AActor>(Object);
+			bAssetsChanged |= Actor == nullptr;
+			bWorldChanged |= Actor != nullptr;
+		}
 
 		TSet<UObject*> SelectedObjectSet( OperationContext->Context->Objects );
 
@@ -602,6 +647,15 @@ void UDataprepActionAsset::ProcessWorkingSetChanged( bool bAddedAssets, bool bAd
 			FDataprepCoreUtils::PurgeObjects( MoveTemp(ObjectsToDelete) );
 		}
 
+		// Build new assets and rebuild modified ones
+		DataprepCorePrivateUtils::ClearAssets( ModifiedAssets.Array() );
+		ModifiedAssets.Reserve( ModifiedAssets.Num() + AddedObjects.Num() );
+		for(UObject* Object : AddedObjects)
+		{
+			ModifiedAssets.Add( Object );
+		}
+		FDataprepCoreUtils::BuildAssets( ModifiedAssets.Array(), ContextPtr->ProgressReporterPtr );
+
 		// Update action's context
 		OperationContext->Context->Objects = SelectedObjectSet.Array();
 
@@ -609,6 +663,10 @@ void UDataprepActionAsset::ProcessWorkingSetChanged( bool bAddedAssets, bool bAd
 		{
 			ContextPtr->ContextChangedCallback( this, bWorldChanged, bAssetsChanged, ContextPtr->Assets.Array() );
 		}
+
+		bWorkingSetHasChanged = false;
+		ModifiedAssets.Empty();
+		AddedObjects.Empty();
 	}
 
 	bWorkingSetHasChanged = false;
diff --git a/Engine/Plugins/Experimental/DataPrepEditor/Source/DataPrepCore/Private/DataprepCorePrivateUtils.cpp b/Engine/Plugins/Experimental/DataPrepEditor/Source/DataPrepCore/Private/DataprepCorePrivateUtils.cpp
--- a/Engine/Plugins/Experimental/DataPrepEditor/Source/DataPrepCore/Private/DataprepCorePrivateUtils.cpp
+++ b/Engine/Plugins/Experimental/DataPrepEditor/Source/DataPrepCore/Private/DataprepCorePrivateUtils.cpp
@@ -4,16 +4,20 @@
 
 #include "DataPrepAsset.h"
 #include "DataprepCoreUtils.h"
+#include "IDataprepProgressReporter.h"
 
 #include "ActorEditorUtils.h"
 #include "AssetRegistryModule.h"
 #include "Engine/Blueprint.h"
+#include "Engine/StaticMesh.h"
 #include "Engine/World.h"
 #include "GameFramework/WorldSettings.h"
 #include "IMessageLogListing.h"
 #include "Kismet2/BlueprintEditorUtils.h"
+#include "MeshDescription.h"
 #include "MessageLogModule.h"
 #include "Modules/ModuleManager.h"
+#include "StaticMeshAttributes.h"
 
 #define LOCTEXT_NAMESPACE "DataprepAsset"
 
@@ -91,4 +95,120 @@ void DataprepCorePrivateUtils::LogMessage( EMessageSeverity::Type Severity, cons
 	}
 }
 
+void DataprepCorePrivateUtils::BuildStaticMeshes(TSet<UStaticMesh*>& StaticMeshes, TFunction<bool(UStaticMesh*)> ProgressFunction, bool bForceBuild)
+{
+	TArray<UStaticMesh*> BuiltMeshes;
+	BuiltMeshes.Reserve( StaticMeshes.Num() );
+
+	if(bForceBuild)
+	{
+		BuiltMeshes.Append( StaticMeshes.Array() );
+	}
+	else
+	{
+		for(UStaticMesh* StaticMesh : StaticMeshes)
+		{
+			if(StaticMesh && (!StaticMesh->RenderData.IsValid() || !StaticMesh->RenderData->IsInitialized()))
+			{
+				BuiltMeshes.Add( StaticMesh );
+			}
+		}
+	}
+
+	if(BuiltMeshes.Num() > 0)
+	{
+		// Start with the biggest mesh first to help balancing tasks on threads
+		BuiltMeshes.Sort(
+			[](const UStaticMesh& Lhs, const UStaticMesh& Rhs) 
+		{ 
+			int32 LhsVerticesNum = Lhs.IsMeshDescriptionValid(0) ? Lhs.GetMeshDescription(0)->Vertices().Num() : 0;
+			int32 RhsVerticesNum = Rhs.IsMeshDescriptionValid(0) ? Rhs.GetMeshDescription(0)->Vertices().Num() : 0;
+
+			return LhsVerticesNum > RhsVerticesNum;
+		}
+		);
+
+		TArray< TArray<FMeshBuildSettings> > StaticMeshesSettings;
+		StaticMeshesSettings.Reserve( BuiltMeshes.Num() );
+
+		//Cache the BuildSettings and update them before building the meshes.
+		for (UStaticMesh* StaticMesh : BuiltMeshes)
+		{
+			TArray<FStaticMeshSourceModel>& SourceModels = StaticMesh->GetSourceModels();
+			TArray<FMeshBuildSettings> BuildSettings;
+			BuildSettings.Reserve(SourceModels.Num());
+
+			for(int32 Index = 0; Index < SourceModels.Num(); ++Index)
+			{
+				FStaticMeshSourceModel& SourceModel = SourceModels[Index];
+
+				BuildSettings.Add( SourceModel.BuildSettings );
+
+				if(FMeshDescription* MeshDescription = StaticMesh->GetMeshDescription(Index))
+				{
+					FStaticMeshAttributes Attributes(*MeshDescription);
+					if(SourceModel.BuildSettings.DstLightmapIndex != -1)
+					{
+						TVertexInstanceAttributesConstRef<FVector2D> VertexInstanceUVs = Attributes.GetVertexInstanceUVs();
+						SourceModel.BuildSettings.bGenerateLightmapUVs = VertexInstanceUVs.IsValid() && VertexInstanceUVs.GetNumIndices() > SourceModel.BuildSettings.DstLightmapIndex;
+					}
+					else
+					{
+						SourceModel.BuildSettings.bGenerateLightmapUVs = false;
+					}
+
+					SourceModel.BuildSettings.bRecomputeNormals = !(Attributes.GetVertexInstanceNormals().IsValid() && Attributes.GetVertexInstanceNormals().GetNumIndices() > 0);
+					SourceModel.BuildSettings.bRecomputeTangents = false;
+					//SourceModel.BuildSettings.bBuildAdjacencyBuffer = false;
+					//SourceModel.BuildSettings.bBuildReversedIndexBuffer = false;
+				}
+			}
+
+			StaticMeshesSettings.Add(MoveTemp(BuildSettings));				
+		}
+
+		// Disable warnings from LogStaticMesh. Not useful
+		ELogVerbosity::Type PrevLogStaticMeshVerbosity = LogStaticMesh.GetVerbosity();
+		LogStaticMesh.SetVerbosity( ELogVerbosity::Error );
+
+		UStaticMesh::BatchBuild(BuiltMeshes, true, ProgressFunction);
+
+		// Restore LogStaticMesh verbosity
+		LogStaticMesh.SetVerbosity( PrevLogStaticMeshVerbosity );
+
+		for(int32 Index = 0; Index < BuiltMeshes.Num(); ++Index)
+		{
+			UStaticMesh* StaticMesh = BuiltMeshes[Index];
+			TArray<FMeshBuildSettings>& PrevBuildSettings = StaticMeshesSettings[Index];
+
+			TArray<FStaticMeshSourceModel>& SourceModels = StaticMesh->GetSourceModels();
+
+			for(int32 SourceModelIndex = 0; SourceModelIndex < SourceModels.Num(); ++SourceModelIndex)
+			{
+				SourceModels[SourceModelIndex].BuildSettings = PrevBuildSettings[SourceModelIndex];
+			}
+
+			if(FStaticMeshRenderData* RenderData = StaticMesh->RenderData.Get())
+			{
+				for ( FStaticMeshLODResources& LODResources : RenderData->LODResources )
+				{
+					LODResources.bHasColorVertexData = true;
+				}
+			}
+		}
+	}
+}
+
+void DataprepCorePrivateUtils::ClearAssets(const TArray<TWeakObjectPtr<UObject>>& Assets)
+{
+	for(const TWeakObjectPtr<UObject>& ObjectPtr : Assets)
+	{
+		if(UStaticMesh* StaticMesh = Cast<UStaticMesh>(ObjectPtr.Get()))
+		{
+			StaticMesh->PreEditChange( nullptr );
+			StaticMesh->RenderData.Reset();
+		}
+	}
+}
+
 #undef LOCTEXT_NAMESPACE
\ No newline at end of file
diff --git a/Engine/Plugins/Experimental/DataPrepEditor/Source/DataPrepCore/Private/DataprepCorePrivateUtils.h b/Engine/Plugins/Experimental/DataPrepEditor/Source/DataPrepCore/Private/DataprepCorePrivateUtils.h
--- a/Engine/Plugins/Experimental/DataPrepEditor/Source/DataPrepCore/Private/DataprepCorePrivateUtils.h
+++ b/Engine/Plugins/Experimental/DataPrepEditor/Source/DataPrepCore/Private/DataprepCorePrivateUtils.h
@@ -7,6 +7,8 @@
 #include "Logging/TokenizedMessage.h"
 
 class AActor;
+class IDataprepProgressReporter;
+class UStaticMesh;
 class UWorld;
 
 namespace DataprepCorePrivateUtils
@@ -70,4 +72,12 @@ namespace DataprepCorePrivateUtils
 	 * @remark Notification is only done ifNotificationText is not empty
 	 */
 	void LogMessage(EMessageSeverity::Type Severity, const FText& Message, const FText& NotificationText = FText());
+
+	/**
+	 * Clear 
+	 */
+	void ClearAssets(const TArray< TWeakObjectPtr< UObject > >& Assets);
+
+	/** Build the render data based on the current geometry available in the static mesh */
+	void BuildStaticMeshes(TSet<UStaticMesh*>& StaticMeshes, TFunction<bool(UStaticMesh*)> ProgressFunction, bool bForceBuild = false);
 }
diff --git a/Engine/Plugins/Experimental/DataPrepEditor/Source/DataPrepCore/Private/DataprepCoreUtils.cpp b/Engine/Plugins/Experimental/DataPrepEditor/Source/DataPrepCore/Private/DataprepCoreUtils.cpp
--- a/Engine/Plugins/Experimental/DataPrepEditor/Source/DataPrepCore/Private/DataprepCoreUtils.cpp
+++ b/Engine/Plugins/Experimental/DataPrepEditor/Source/DataPrepCore/Private/DataprepCoreUtils.cpp
@@ -10,11 +10,17 @@
 
 #include "DataPrepAsset.h"
 #include "DataprepCoreLogCategory.h"
+#include "DataprepCorePrivateUtils.h"
 #include "IDataprepProgressReporter.h"
 
 #include "Engine/StaticMesh.h"
 #include "GameFramework/Actor.h"
 #include "LevelSequence.h"
+#include "Materials/MaterialFunction.h"
+#include "Materials/MaterialInstance.h"
+#include "Materials/Material.h"
+#include "MaterialGraph/MaterialGraph.h"
+#include "MaterialShared.h"
 #include "Misc/ScopedSlowTask.h"
 #include "RenderingThread.h"
 #include "UObject/UObjectHash.h"
@@ -304,4 +310,61 @@ FFeedbackContext* FDataprepCoreUtils::FDataprepProgressTextReporter::GetFeedback
 	return FeedbackContext.Get();
 }
 
+void FDataprepCoreUtils::BuildAssets(const TArray<TWeakObjectPtr<UObject>>& Assets, const TSharedPtr<IDataprepProgressReporter>& ProgressReporterPtr)
+{
+	TSet<UStaticMesh*> StaticMeshes;
+	TSet<UMaterialInterface*> MaterialInterfaces;
+
+	for( const TWeakObjectPtr<UObject>& AssetPtr : Assets )
+	{
+		UObject* AssetObject = AssetPtr.Get();
+		if( AssetObject )
+		{
+			if(UMaterialInterface* MaterialInterface = Cast<UMaterialInterface>(AssetObject))
+			{
+				MaterialInterfaces.Add( MaterialInterface );
+			}
+			else if(UStaticMesh* StaticMesh = Cast<UStaticMesh>(AssetObject))
+			{
+				StaticMeshes.Add( StaticMesh );
+			}
+		}
+	}
+
+	// Force compilation of materials which have no render proxy
+	if(MaterialInterfaces.Num() > 0)
+	{
+		auto MustCompile = [](const UMaterialInterface* MaterialInterface)
+		{
+			if( MaterialInterface )
+			{
+				const FMaterialRenderProxy* RenderProxy = MaterialInterface->GetRenderProxy();
+				return RenderProxy == nullptr || !RenderProxy->IsInitialized();
+			}
+
+			return false;
+		};
+
+		for(UMaterialInterface* MaterialInterface : MaterialInterfaces)
+		{
+			if( MustCompile( MaterialInterface ) )
+			{
+				FPropertyChangedEvent EmptyPropertyUpdateStruct( nullptr );
+				MaterialInterface->PostEditChangeProperty( EmptyPropertyUpdateStruct );
+
+			}
+
+			if( MustCompile( MaterialInterface->GetMaterial() ) )
+			{
+				FPropertyChangedEvent EmptyPropertyUpdateStruct( nullptr );
+				MaterialInterface->GetMaterial()->PostEditChangeProperty( EmptyPropertyUpdateStruct );
+
+			}
+		}
+	}
+
+	// Build static meshes
+	DataprepCorePrivateUtils::BuildStaticMeshes( StaticMeshes, TFunction<bool(UStaticMesh*)>() );
+}
+
 #undef LOCTEXT_NAMESPACE
\ No newline at end of file
diff --git a/Engine/Plugins/Experimental/DataPrepEditor/Source/DataPrepCore/Public/DataPrepOperation.h b/Engine/Plugins/Experimental/DataPrepEditor/Source/DataPrepCore/Public/DataPrepOperation.h
--- a/Engine/Plugins/Experimental/DataPrepEditor/Source/DataPrepCore/Public/DataPrepOperation.h
+++ b/Engine/Plugins/Experimental/DataPrepEditor/Source/DataPrepCore/Public/DataPrepOperation.h
@@ -13,12 +13,20 @@
 #include "DataPrepOperation.generated.h"
 
 /**
- * Delegate used by a Dataprep operation to add an asset to the Dataprep's working set
+ * Delegate used by a Dataprep operation to add a copy of an asset to the Dataprep's working set
  * @return	A new asset which is either a duplicate of Asset if not null or of class AssetClass and name AssetName
  * @remark	Returns nullptr if Asset and AssetClass is null
  * @remark	If Asset is not null, it will not be added to the list of assets tracked by Dataprep. It can safely be deleted
  */
-DECLARE_DELEGATE_RetVal_ThreeParams(UObject* /* NewAsset */, FDataprepAddAsset, const UObject* /* Asset */, UClass* /* AssetClass */, const TCHAR* /* Assetname */ )
+DECLARE_DELEGATE_RetVal_TwoParams(UObject* /* NewAsset */, FDataprepAddAsset, const UObject* /* Asset */, const TCHAR* /* Assetname */ )
+
+/**
+ * Delegate used by a Dataprep operation to create and add an asset to the Dataprep's working set
+ * @return	A new asset which is either a duplicate of Asset if not null or of class AssetClass and name AssetName
+ * @remark	Returns nullptr if Asset and AssetClass is null
+ * @remark	If Asset is not null, it will not be added to the list of assets tracked by Dataprep. It can safely be deleted
+ */
+DECLARE_DELEGATE_RetVal_TwoParams(UObject* /* NewAsset */, FDataprepCreateAsset, UClass* /* AssetClass */, const TCHAR* /* Assetname */ )
 
 /**
  * Delegate used by a Dataprep operation to add an actor to the Dataprep's working set
@@ -34,6 +42,11 @@ DECLARE_DELEGATE_RetVal_TwoParams(class AActor* /* Actor */, FDataprepCreateActo
  */
 DECLARE_DELEGATE_TwoParams(FDataprepRemoveObject, UObject* /* Object */, bool /* bLocalContext */ )
 
+/**
+ * Delegate used by a Dataprep operation to indicates an asset has been modified
+ */
+DECLARE_DELEGATE_OneParam(FDataprepAssetsModified, TArray<UObject*> /* Objects */ )
+
 /**
  * Delegate used by a Dataprep operation to remove and delete an array of objects from the Dataprep's working set.
  * After execution, Object will not be accessible by subsequent operations.
@@ -74,13 +87,24 @@ struct FDataprepOperationContext
 	TSharedPtr<struct FDataprepContext> Context;
 
 	/**
-	 * Delegate to add an asset to the Dataprep content
+	 * Delegate to indicate an asset has been modified
+	 */
+	FDataprepAssetsModified AssetsModifiedDelegate;
+
+	/**
+	 * Delegate to duplicate and add an asset to the Dataprep content
 	 * Only effective if the operation is of class UDataprepEditingOperation
 	 */
 	FDataprepAddAsset AddAssetDelegate;
 
 	/**
-	 * Delegate to add an asset to the Dataprep content
+	 * Delegate to create and add an asset to the Dataprep content
+	 * Only effective if the operation is of class UDataprepEditingOperation
+	 */
+	FDataprepCreateAsset CreateAssetDelegate;
+
+	/**
+	 * Delegate to create and add an actor to the Dataprep content
 	 * Only effective if the operation is of class UDataprepEditingOperation
 	 */
 	FDataprepCreateActor CreateActorDelegate;
@@ -178,6 +202,14 @@ protected:
 	UFUNCTION(BlueprintCallable,  Category = "Report", meta = (HideSelfPin = "true"))
 	void ReportProgress( float IncrementOfWork, const FText& InMessage );
 
+	/**
+	 * Indicates an array of assets has changed during the operation. It is important to use this function
+	 * if the modifications on the assets impact their appearance
+	 * @param Assets			Array of assets which have been modified
+	 */
+	UFUNCTION(BlueprintCallable,  Category = "Dataprep | Operation")
+	void AssetsModified( TArray<UObject*> Assets );
+
 	/**
 	 * Create a task to report progress during the execution of an operation
 	 * @param InDescription		The description of the task about to be performed
@@ -186,6 +218,9 @@ protected:
 	 */
 	TSharedPtr<FDataprepWorkReporter> CreateTask( const FText& InDescription, float InAmountOfWork, float InIncrementOfWork = 1.0f );
 
+	/** Returns true if the operation was canceled during execution */
+	bool IsCancelled();
+
 	// User friendly interface end here ========================================================================
 
 public:
@@ -241,14 +276,22 @@ class DATAPREPCORE_API UDataprepEditingOperation : public UDataprepOperation
 
 protected:
 	/**
-	 * Add an asset to the Dataprep's and action's working set
+	 * DUplicate and add an asset to the Dataprep's and action's working set
 	 * @param Asset			If not null, the asset will be duplicated
+	 * @param AssetName		Name of the asset to create. Name collision will be checked and fixed before naming the asset
+	 * @returns				The asset newly created
+	 */
+	UFUNCTION(BlueprintCallable,  Category = "Dataprep | Editing Operation")
+	UObject* AddAsset(const UObject* Asset, const FString& AssetName );
+
+	/**
+	 * Create and add an asset to the Dataprep's and action's working set
 	 * @param AssetClass	If Asset is null, an asset of the given class will be returned
-	 * @param AssetName		Name of the asset to create. Name collision will be performed before naming the asset
+	 * @param AssetName		Name of the asset to create. Name collision will be checked and fixed before naming the asset
 	 * @returns				The asset newly created
 	 */
 	UFUNCTION(BlueprintCallable,  Category = "Dataprep | Editing Operation")
-	UObject* AddAsset(const UObject* Asset, UClass* AssetClass, const FString& AssetName );
+	UObject* CreateAsset(UClass* AssetClass, const FString& AssetName );
 
 	/**
 	 * Add an actor to the Dataprep's transient world and action's working set
diff --git a/Engine/Plugins/Experimental/DataPrepEditor/Source/DataPrepCore/Public/DataprepActionAsset.h b/Engine/Plugins/Experimental/DataPrepEditor/Source/DataPrepCore/Public/DataprepActionAsset.h
--- a/Engine/Plugins/Experimental/DataPrepEditor/Source/DataPrepCore/Public/DataprepActionAsset.h
+++ b/Engine/Plugins/Experimental/DataPrepEditor/Source/DataPrepCore/Public/DataprepActionAsset.h
@@ -262,13 +262,20 @@ private:
 	void RemoveInvalidOperations();
 
 	/**
-	 * Add an asset to the Dataprep's and action's working set
+	 * Duplicate and add an asset to the Dataprep's and action's working set
 	 * @param Asset			If not null, the asset will be duplicated
+	 * @param AssetName		Name of the asset to create. Name collision will be performed before naming the asset
+	 * @returns				The asset newly created
+	 */
+	UObject* OnAddAsset(const UObject* Asset, const TCHAR* AssetName);
+
+	/**
+	 * Create and add an asset to the Dataprep's and action's working set
 	 * @param AssetClass	If Asset is null, an asset of the given class will be returned
 	 * @param AssetName		Name of the asset to create. Name collision will be performed before naming the asset
 	 * @returns				The asset newly created
 	 */
-	UObject* OnAddAsset(const UObject* Asset, UClass* AssetClass, const TCHAR* AssetName);
+	UObject* OnCreateAsset(UClass* AssetClass, const TCHAR* AssetName);
 
 	/**
 	 * Add an actor to the Dataprep's transient world and action's working set
@@ -288,6 +295,12 @@ private:
 	 */
 	void OnRemoveObject(UObject* Object, bool bLocalContext);
 
+	/**
+	 * Add an array of assets to the list of modified assets
+	 * @param Assets		An array of assets which have been modifed
+	 */
+	void OnAssetsModified(TArray<UObject*> Assets);
+
 	/**
 	 * Delete an array of objects from the Dataprep's and action's working set
 	 * @param Objects		The array of objects to delete
@@ -300,14 +313,16 @@ private:
 
 	/**
 	 * Executes the deletion and removal requested by an operation after its execution
-	 * @param bAddedAssets		Indicates that assets have been added during the operation's execution
-	 * @param bAddedActors		Indicates that actors have been added during the operation's execution
+	 * Notifies any observer about what has changed
 	 */
-	void ProcessWorkingSetChanged(bool bAddedAssets, bool bAddedActors);
+	void ProcessWorkingSetChanged();
 
 	/** Returns the outer to be used according to an asset's class */
 	UObject* GetAssetOuterByClass( UClass* AssetClass );
 
+	/** Add an asset to the execution context */
+	void AddAssetToContext( UObject* NewAsset, const TCHAR* DesiredName );
+
 	/** Array of operations and/or filters constituting this action */
 	UPROPERTY()
 	TArray<UDataprepActionStep*> Steps;
@@ -326,6 +341,12 @@ private:
 	/** Array of objects requested to be deleted by an operation */
 	TArray<UObject*> ObjectsToDelete;
 
+	/** Set of objects which have been modified during the execution of an operation */
+	TSet< TWeakObjectPtr< UObject > > ModifiedAssets;
+
+	/** Array of objects which have been added during the execution of an operation */
+	TArray<UObject*> AddedObjects;
+
 	/** Array of objects requested to be removed by an operation */
 	TArray<TPair<UObject*,bool>> ObjectsToRemove;
 
diff --git a/Engine/Plugins/Experimental/DataPrepEditor/Source/DataPrepCore/Public/DataprepCoreUtils.h b/Engine/Plugins/Experimental/DataPrepEditor/Source/DataPrepCore/Public/DataprepCoreUtils.h
--- a/Engine/Plugins/Experimental/DataPrepEditor/Source/DataPrepCore/Public/DataprepCoreUtils.h
+++ b/Engine/Plugins/Experimental/DataPrepEditor/Source/DataPrepCore/Public/DataprepCoreUtils.h
@@ -69,6 +69,14 @@ public:
 		}
 	}
 
+	/**
+	 * Helper function to build assets for use in the Dataprep pipeline
+
+	 * @param	Assets					Array of weak pointer on the assets to be build
+	 * @param	ProgressReporterPtr		Pointer to a IDataprepProgressReporter interface. This pointer can be invalid.
+	 */
+	static void BuildAssets(const TArray< TWeakObjectPtr<UObject> >& Assets, const TSharedPtr< IDataprepProgressReporter >& ProgressReporterPtr );
+
 	class DATAPREPCORE_API FDataprepLogger : public IDataprepLogger
 	{
 	public:
diff --git a/Engine/Plugins/Experimental/DataPrepEditor/Source/DataPrepEditor/Private/DataPrepEditor.cpp b/Engine/Plugins/Experimental/DataPrepEditor/Source/DataPrepEditor/Private/DataPrepEditor.cpp
--- a/Engine/Plugins/Experimental/DataPrepEditor/Source/DataPrepEditor/Private/DataPrepEditor.cpp
+++ b/Engine/Plugins/Experimental/DataPrepEditor/Source/DataPrepEditor/Private/DataPrepEditor.cpp
@@ -660,7 +660,7 @@ void FDataprepEditor::CleanPreviewWorld()
 		}
 	}
 
-	SceneViewportView->ClearMeshes();
+	SceneViewportView->ClearScene();
 
 	// Delete assets which are still in the transient content folder
 	const FString TransientContentFolder( GetTransientContentFolder() );
@@ -708,9 +708,8 @@ void FDataprepEditor::OnExecutePipeline()
 		RestoreFromSnapshot();
 	}
 
-	// Delete any reference of assets in the 3D viewport
-	// Required because assets could be removed and the 3D viewport is using raw pointers
-	SceneViewportView->ClearMeshes();
+	// Remove any link between assets referenced in the preview world and the viewport
+	SceneViewportView->ClearScene();
 
 	// Trigger execution of data preparation operations on world attached to recipe
 	{
@@ -736,7 +735,7 @@ void FDataprepEditor::OnExecutePipeline()
 	}
 
 	// Redraw 3D viewport
-	SceneViewportView->UpdateMeshes();
+	SceneViewportView->UpdateScene();
 
 	// Add assets which may have been created by actions
 	for( TWeakObjectPtr<UObject>& Asset : Assets )
@@ -783,7 +782,7 @@ void FDataprepEditor::OnCommitWorld()
 	}
 
 	// Remove references to assets in 3D viewport before commit
-	SceneViewportView->ClearMeshes();
+	SceneViewportView->ClearScene();
 
 	// Finalize assets
 	TArray<TWeakObjectPtr<UObject>> ValidAssets( MoveTemp(Assets) );
@@ -1158,7 +1157,7 @@ void FDataprepEditor::UpdatePreviewPanels(bool bInclude3DViewport)
 
 	if(bInclude3DViewport)
 	{
-		SceneViewportView->UpdateMeshes();
+		SceneViewportView->UpdateScene();
 	}
 }
 
@@ -1193,9 +1192,6 @@ bool FDataprepEditor::OnCanExecuteNextStep(UDataprepActionAsset* ActionAsset, UD
 	// #ueent_todo: Make this action configurable by user
 	UpdatePreviewPanels(false);
 
-	//// Remove references to assets in 3D viewport in case they are deleted in next step
-	//SceneViewportView->ClearMeshes();
-
 	// #ueent_todo: Make this action configurable by user
 	return true;
 }
diff --git a/Engine/Plugins/Experimental/DataPrepEditor/Source/DataPrepEditor/Private/DataprepSnapshot.cpp b/Engine/Plugins/Experimental/DataPrepEditor/Source/DataPrepEditor/Private/DataprepSnapshot.cpp
--- a/Engine/Plugins/Experimental/DataPrepEditor/Source/DataPrepEditor/Private/DataprepSnapshot.cpp
+++ b/Engine/Plugins/Experimental/DataPrepEditor/Source/DataPrepEditor/Private/DataprepSnapshot.cpp
@@ -688,14 +688,16 @@ void FDataprepEditor::RestoreFromSnapshot(bool bUpdateViewport)
 		FDataprepCoreUtils::PurgeObjects( MoveTemp( ObjectsToDelete ) );
 	}
 
-	/** 
-	 * Set all assets as public so the actors in the level can find the assets they are referring to
-	 * Also assets should always have a RF_Public flags as they are the public interface of their package
-	 */
+	// Make sure all assets have RF_Public flag set so the actors in the level can find the assets they are referring to
+	// Cache boolean representing if RF_Public flag was set or not on asset
+	TArray<bool> AssetFlags;
+	AssetFlags.AddDefaulted( Assets.Num() );
+
 	for(int32 Index = 0; Index < Assets.Num(); ++Index)
 	{
 		if( UObject* Asset = Assets[Index].Get() )
 		{
+			AssetFlags[Index] = bool(Asset->GetFlags() & RF_Public);
 			Asset->SetFlags( RF_Public );
 		}
 	}
@@ -750,6 +752,23 @@ void FDataprepEditor::RestoreFromSnapshot(bool bUpdateViewport)
 	}
 	UE_LOG( LogDataprepEditor, Verbose, TEXT("Level loaded") );
 
+	// Restore RF_Public on each asset
+	for(int32 Index = 0; Index < Assets.Num(); ++Index)
+	{
+		if( UObject* Asset = Assets[Index].Get() )
+		{
+			if( !AssetFlags[Index] )
+			{
+				Asset->ClearFlags( RF_Public );
+			}
+		}
+	}
+
+	{
+		TSharedPtr< IDataprepProgressReporter > ProgressReporter( new FDataprepCoreUtils::FDataprepProgressUIReporter() );
+		FDataprepCoreUtils::BuildAssets( Assets, ProgressReporter );
+	}
+
 	// Log time spent to import incoming file in minutes and seconds
 	double ElapsedSeconds = FPlatformTime::ToSeconds64(FPlatformTime::Cycles64() - StartTime);
 
diff --git a/Engine/Plugins/Experimental/DataPrepEditor/Source/DataPrepEditor/Private/Widgets/SDataprepEditorViewport.cpp b/Engine/Plugins/Experimental/DataPrepEditor/Source/DataPrepEditor/Private/Widgets/SDataprepEditorViewport.cpp
--- a/Engine/Plugins/Experimental/DataPrepEditor/Source/DataPrepEditor/Private/Widgets/SDataprepEditorViewport.cpp
+++ b/Engine/Plugins/Experimental/DataPrepEditor/Source/DataPrepEditor/Private/Widgets/SDataprepEditorViewport.cpp
@@ -4,6 +4,7 @@
 
 #include "DataprepCoreUtils.h"
 #include "DataPrepEditor.h"
+#include "DataprepEditorUtils.h"
 
 #include "DataprepEditorLogCategory.h"
 
@@ -41,9 +42,6 @@
 
 #define LOCTEXT_NAMESPACE "DataprepEditorViewport"
 
-// #ueent_todo: Boolean to locally enable multi-threaded build of meshes while the discussion about the proper solution is going on
-static bool bComputeUVStretching = false;
-
 TWeakObjectPtr<UMaterial> SDataprepEditorViewport::PreviewMaterial;
 TWeakObjectPtr<UMaterial> SDataprepEditorViewport::XRayMaterial;
 TWeakObjectPtr<UMaterial> SDataprepEditorViewport::BackFaceMaterial;
@@ -179,8 +177,6 @@ namespace DataprepEditor3DPreviewUtils
 
 	void FindMeshComponents(const AActor * InActor, TArray<UStaticMeshComponent*>& MeshComponents, bool bRecursive );
 
-	void ComputeUVStretching(FStaticMeshLODResources&  Resource);
-
 	/** Returns array of static mesh components in world */
 	template<class T>
 	TArray< T* > GetComponentsFromWorld( UWorld* World )
@@ -239,7 +235,7 @@ SDataprepEditorViewport::SDataprepEditorViewport()
 SDataprepEditorViewport::~SDataprepEditorViewport()
 {
 	CastChecked<UEditorEngine>(GEngine)->OnPreviewFeatureLevelChanged().Remove(PreviewFeatureLevelChangedHandle);
-	ClearMeshes();
+	ClearScene();
 }
 
 void SDataprepEditorViewport::Construct(const FArguments& InArgs, TSharedPtr<FDataprepEditor> InDataprepEditor)
@@ -277,7 +273,7 @@ void SDataprepEditorViewport::Construct(const FArguments& InArgs, TSharedPtr<FDa
 	SEditorViewport::Construct( SEditorViewport::FArguments() );
 }
 
-void SDataprepEditorViewport::ClearMeshes()
+void SDataprepEditorViewport::ClearScene()
 {
 	const int32 PreviousCount = PreviewMeshComponents.Num();
 	TArray<UObject*> ObjectsToDelete;
@@ -287,49 +283,26 @@ void SDataprepEditorViewport::ClearMeshes()
 	{
 		if(UStaticMeshComponent* MeshComponent = PreviewMeshComponent.Get())
 		{
+			MeshComponent->EmptyOverrideMaterials();
 			ObjectsToDelete.Add( MeshComponent );
 		}
 	}
 
 	FDataprepCoreUtils::PurgeObjects( MoveTemp( ObjectsToDelete ) );
 
-	// Release render data created for display
-	for(UStaticMesh* StaticMesh : BuiltMeshes)
-	{
-		if(StaticMesh)
-		{
-			// Free any RHI resources created for display
-			StaticMesh->PreEditChange(nullptr);
-			StaticMesh->RenderData.Reset();
-			// No need to post-edit
-		}
-	}
-	BuiltMeshes.Empty( BuiltMeshes.Num() );
-
-	// Restoring mesh components' render states
-	for( TWeakObjectPtr< UStaticMeshComponent >& MeshComponentPtr : MeshComponentsToRestore )
-	{
-		if(UStaticMeshComponent* MeshComponent = MeshComponentPtr.Get())
-		{
-			MeshComponent->RecreateRenderState_Concurrent();
-		}
-	}
-	MeshComponentsToRestore.Empty( MeshComponentsToRestore.Num() );
-
 	PreviewMeshComponents.Empty(PreviousCount);
 	MeshComponentsMapping.Empty(PreviousCount);
 	MeshComponentsReverseMapping.Empty(PreviousCount);
-	DisplayMaterialsMap.Empty(DisplayMaterialsMap.Num());
 	SelectedPreviewComponents.Empty(SelectedPreviewComponents.Num());
 
 	OverlayTextVerticalBox->ClearChildren();
 }
 
-void SDataprepEditorViewport::UpdateMeshes()
+void SDataprepEditorViewport::UpdateScene()
 {
-	TRACE_CPUPROFILER_EVENT_SCOPE(SDataprepEditorViewport::UpdateMeshes);
+	TRACE_CPUPROFILER_EVENT_SCOPE(SDataprepEditorViewport::UpdateScene);
 
-	ClearMeshes();
+	ClearScene();
 
 	RenderingMaterialType = ERenderingMaterialType::OriginalRenderingMaterial;
 	bWireframeRenderingMode = false;
@@ -344,65 +317,36 @@ void SDataprepEditorViewport::UpdateMeshes()
 
 	if(SceneMeshComponents.Num() > 0)
 	{
-		TSet<UStaticMesh*> StaticMeshes;
+		int32 InvalidStaticMeshesCount = 0;
 
 		for( UStaticMeshComponent*& MeshComponent : SceneMeshComponents )
 		{
 			if( UStaticMesh* StaticMesh = MeshComponent->GetStaticMesh() )
 			{
-				StaticMeshes.Add( StaticMesh );
+				if( !StaticMesh->RenderData.IsValid() || !StaticMesh->RenderData->IsInitialized())
+				{
+					++InvalidStaticMeshesCount;
+					MeshComponent = nullptr;
+				}
 			}
 			else
 			{
+				++InvalidStaticMeshesCount;
 				MeshComponent = nullptr;
 			}
 		}
 
-		if( StaticMeshes.Num() > 0)
+		if( SceneMeshComponents.Num() > InvalidStaticMeshesCount )
 		{
 			FScopedSlowTask SlowTask( 100.0f, LOCTEXT("UpdateMeshes_Title", "Updating 3D viewport ...") );
 			SlowTask.MakeDialog(false);
 
-			// Build render data of static meshes for display
-			SlowTask.EnterProgressFrame( 50.0f, LOCTEXT("UpdateMeshes_StaticMeshes", "Buidling static meshes ...") );
-			DataprepEditor3DPreviewUtils::BuildStaticMeshes( StaticMeshes, BuiltMeshes );
-
-			// Clear render state of static mesh components from the preview world which
-			// static meshes have been built for the 3D viewport
-			// Required so mesh components from the preview world are not impacted by the creation
-			// and deletion of render data done for the viewport
-			// The mesh components' render state will be restored when the viewport is cleared
-			{
-				TSet<UStaticMesh*> BuiltMeshesSet;
-				BuiltMeshesSet.Reserve( BuiltMeshes.Num() );
-				BuiltMeshesSet.Append( BuiltMeshes );
-
-				MeshComponentsToRestore.Empty( SceneMeshComponents.Num() );
-				for( UStaticMeshComponent*& MeshComponent : SceneMeshComponents )
-				{
-					if( MeshComponent->IsRegistered() && BuiltMeshesSet.Contains( MeshComponent->GetStaticMesh() ) )
-					{
-						if( MeshComponent->IsRenderStateCreated() )
-						{
-							if( !MeshComponent->IsRenderStateDirty() )
-							{
-								MeshComponent->DoDeferredRenderUpdates_Concurrent();
-							}
-
-							MeshComponent->DestroyRenderState_Concurrent();
-
-							MeshComponentsToRestore.Add( MeshComponent );
-						}
-					}
-				}
-			}
-
-			CreateDisplayMaterials( SceneMeshComponents );
+			InitializeDefaultMaterials();
 
 			PreviewMeshComponents.Empty( SceneMeshComponents.Num() );
 
 			// Compute bounding box of scene to determine camera position and scaling to apply for smooth navigation
-			SlowTask.EnterProgressFrame( 10.0f, LOCTEXT("UpdateMeshes_Prepare", "Preparing viewport ...") );
+			SlowTask.EnterProgressFrame( 10.0f, LOCTEXT("UpdateMeshes_Prepare", "Computing scene's bounding box ...") );
 			SceneBounds.Init();
 			for( UStaticMeshComponent* SceneMeshComponent : SceneMeshComponents )
 			{
@@ -429,21 +373,17 @@ void SDataprepEditorViewport::UpdateMeshes()
 			const int32 PerMeshColorsCount = sizeof(PerMeshColor);
 			int32 PerMeshColorIndex = 0;
 			// Replicate mesh component from world to preview in preview world
-			SlowTask.EnterProgressFrame( 40.0f, LOCTEXT("UpdateMeshes_Components", "Adding meshes to viewport ...") );
+			SlowTask.EnterProgressFrame( 90.0f, LOCTEXT("UpdateMeshes_Components", "Adding meshes to viewport ...") );
 			{
-				FScopedSlowTask SubSlowTask( SceneMeshComponents.Num(), LOCTEXT("UpdateMeshes_Components", "Adding meshes to viewport ...") );
-				SubSlowTask.MakeDialog(false);
+				FScopedSlowTask SubSlowTask( float(SceneMeshComponents.Num() - InvalidStaticMeshesCount), LOCTEXT("UpdateMeshes_Components", "Adding meshes to viewport ...") );
+
 				for( UStaticMeshComponent* SceneMeshComponent : SceneMeshComponents )
 				{
-					FText Message = LOCTEXT("UpdateMeshes_SkipOneComponent", "Skipping null actor ...");
-					if(SceneMeshComponent != nullptr && SceneMeshComponent->GetOwner() )
-					{
-						Message = FText::Format( LOCTEXT("UpdateMeshes_AddOneComponent", "Adding {0} ..."), FText::FromString( SceneMeshComponent->GetOwner()->GetActorLabel() ) );
-					}
-					SubSlowTask.EnterProgressFrame( 1.0f, Message );
-
 					if(SceneMeshComponent != nullptr)
 					{
+						const FText Message = FText::Format( LOCTEXT("UpdateMeshes_AddOneComponent", "Adding {0} ..."), FText::FromString( SceneMeshComponent->GetOwner()->GetActorLabel() ) );
+						SubSlowTask.EnterProgressFrame( 1.0f, Message );
+						
 						UCustomStaticMeshComponent* PreviewMeshComponent = NewObject< UCustomStaticMeshComponent >( PreviewActor.Get(), NAME_None, RF_Transient );
 						if (GEditor->PreviewPlatform.GetEffectivePreviewFeatureLevel() <= ERHIFeatureLevel::ES3_1)
 						{
@@ -463,7 +403,6 @@ void SDataprepEditorViewport::UpdateMeshes()
 
 						PreviewMeshComponent->AttachToComponent( PreviewActor->GetRootComponent(), FAttachmentTransformRules::KeepRelativeTransform );
 						PreviewMeshComponent->SetRelativeTransform( ComponentToWorldTransform );
-						PreviewMeshComponent->RegisterComponentWithWorld( PreviewScene->GetWorld() );
 
 						// Apply preview material to preview static mesh component
 						for(int32 Index = 0; Index < StaticMesh->StaticMaterials.Num(); ++Index)
@@ -475,11 +414,13 @@ void SDataprepEditorViewport::UpdateMeshes()
 								MaterialInterface = StaticMesh->StaticMaterials[Index].MaterialInterface;
 							}
 
-							PreviewMeshComponent->SetMaterial( Index, DisplayMaterialsMap[MaterialInterface].Get() );
+							PreviewMeshComponent->SetMaterial( Index, MaterialInterface );
 						}
 
 						PreviewMeshComponent->SelectionOverrideDelegate = UPrimitiveComponent::FSelectionOverride::CreateRaw( this, &SDataprepEditorViewport::IsComponentSelected );
 
+						PreviewMeshComponent->RegisterComponentWithWorld( PreviewScene->GetWorld() );
+
 						PreviewMeshComponents.Emplace( PreviewMeshComponent );
 
 						MeshComponentsMapping.Add( SceneMeshComponent, PreviewMeshComponent );
@@ -492,6 +433,15 @@ void SDataprepEditorViewport::UpdateMeshes()
 			{
 				ViewportDebug::FTimeLogger LapTimeLogger( TEXT("Building mesh properties") );
 
+				TSet<UStaticMesh*> StaticMeshes;
+				for( const TWeakObjectPtr< UStaticMeshComponent >& PreviewMeshComponentPtr : PreviewMeshComponents )
+				{
+					if(UCustomStaticMeshComponent* PreviewMeshComponent = Cast<UCustomStaticMeshComponent>(PreviewMeshComponentPtr.Get()))
+					{
+						StaticMeshes.Add( PreviewMeshComponent->GetStaticMesh() );
+					}
+				}
+
 				TMap<UStaticMesh*, FPrototypeOrientedBox> MeshPropertiesMap;
 				TArray<UStaticMesh*> StaticMeshesToBuild = StaticMeshes.Array();
 
@@ -501,7 +451,7 @@ void SDataprepEditorViewport::UpdateMeshes()
 					MeshProperties.AddDefaulted( StaticMeshes.Num() );
 
 					ParallelFor( StaticMeshesToBuild.Num(), [&]( int32 Index ) {
-						MeshDescriptionPrototype::GenerateOrientedBox( *StaticMeshesToBuild[Index]->GetMeshDescription(0), MeshProperties[Index], StaticMeshesToBuild[Index]->GetName());
+						MeshProperties[Index] = MeshDescriptionPrototype::GenerateOrientedBox( StaticMeshesToBuild[Index] );
 					});
 
 					for(int32 Index = 0; Index < StaticMeshesToBuild.Num(); ++Index)
@@ -511,7 +461,7 @@ void SDataprepEditorViewport::UpdateMeshes()
 				}
 				else
 				{
-					MeshPropertiesMap.Add( StaticMeshesToBuild[0], MeshDescriptionPrototype::GenerateOrientedBox( *StaticMeshesToBuild[0]->GetMeshDescription(0), StaticMeshesToBuild[0]->GetName() ) );
+					MeshPropertiesMap.Add( StaticMeshesToBuild[0], MeshDescriptionPrototype::GenerateOrientedBox( StaticMeshesToBuild[0] ) );
 				}
 
 				TSet<UStaticMesh*> ShouldBeInstanced;
@@ -870,82 +820,6 @@ void SDataprepEditorViewport::InitializeDefaultMaterials()
 	}
 }
 
-void SDataprepEditorViewport::CreateDisplayMaterials(const TArray< UStaticMeshComponent* >& SceneMeshComponents)
-{
-	InitializeDefaultMaterials();
-
-	DisplayMaterialsMap.Empty();
-
-	for( const UStaticMeshComponent* SceneMeshComponent : SceneMeshComponents )
-	{
-		if(SceneMeshComponent != nullptr)
-		{
-			UStaticMesh* StaticMesh = SceneMeshComponent->GetStaticMesh();
-
-			for(int32 Index = 0; Index < StaticMesh->StaticMaterials.Num(); ++Index)
-			{
-				if(UMaterialInterface* MaterialInterface = SceneMeshComponent->GetMaterial(Index))
-				{
-					DisplayMaterialsMap.Add( MaterialInterface );
-				}
-				else
-				{
-					DisplayMaterialsMap.Add( StaticMesh->StaticMaterials[Index].MaterialInterface );
-				}
-			}
-		}
-	}
-
-	TArray< UMaterialInterface* > Materials;
-	Materials.Reserve( DisplayMaterialsMap.Num() );
-
-	TMap< UMaterialInterface*, UMaterialInterface* > ParentMaterials;
-
-	for(auto& MaterialEntry : DisplayMaterialsMap)
-	{
-		UMaterialInstanceConstant* MaterialInstance = nullptr;
-		if(UMaterialInstanceConstant* ConstantMaterialInstance = Cast< UMaterialInstanceConstant >(MaterialEntry.Key))
-		{
-			MaterialInstance = DuplicateObject< UMaterialInstanceConstant >( ConstantMaterialInstance, GetTransientPackage(), NAME_None);
-			if( !ParentMaterials.Contains(ConstantMaterialInstance->Parent) )
-			{
-				// #ueent_todo: Assuming here that the parent is a Material
-				UMaterial* SourceMaterial = Cast<UMaterial>(ConstantMaterialInstance->Parent);
-				ensure(SourceMaterial);
-
-				UMaterial* ParentMaterial = DuplicateObject< UMaterial >( SourceMaterial, GetTransientPackage(), NAME_None);
-
-				ParentMaterials.Add( SourceMaterial, ParentMaterial );
-				Materials.Add( ParentMaterial );
-			}
-
-			MaterialInstance->SetFlags( RF_Transient );
-			MaterialInstance->Parent = ParentMaterials[ConstantMaterialInstance->Parent];
-		}
-		else if(UMaterial* Material = Cast<UMaterial>(MaterialEntry.Key))
-		{
-			UMaterial* ParentMaterial = DuplicateObject< UMaterial >( Material, GetTransientPackage(), NAME_None);
-
-			ParentMaterials.Add( Material, ParentMaterial );
-			Materials.Add( ParentMaterial );
-
-			MaterialInstance = NewObject<UMaterialInstanceConstant>( GetTransientPackage(), NAME_None, EObjectFlags::RF_Transient);
-			MaterialInstance->Parent = ParentMaterial;
-		}
-		else
-		{
-			MaterialInstance = NewObject<UMaterialInstanceConstant>( GetTransientPackage(), NAME_None, EObjectFlags::RF_Transient);
-			MaterialInstance->Parent = PreviewMaterial.Get();
-		}
-
-		MaterialEntry.Value = TWeakObjectPtr<UMaterialInstanceConstant>(MaterialInstance);
-		Materials.Add( MaterialInstance );
-	}
-
-
-	DataprepEditor3DPreviewUtils::CompileMaterials( Materials );
-}
-
 void SDataprepEditorViewport::SetSelection(UStaticMeshComponent* Component)
 {
 	SelectedPreviewComponents.Empty( 1 );
@@ -1072,7 +946,7 @@ void SDataprepEditorViewport::UpdateSelection()
 	// Apply materials. Only selected ones will be affected
 	ApplyRenderingMaterial();
 
-	// Update Dataprep editor with new selection
+	// UpdateScene Dataprep editor with new selection
 	TSet<TWeakObjectPtr<UObject>> SelectedActors;
 	SelectedActors.Empty(SelectedPreviewComponents.Num());
 
@@ -1187,7 +1061,7 @@ void SDataprepEditorViewport::ApplyRenderingMaterial()
 				MaterialInterface = StaticMesh->StaticMaterials[Index].MaterialInterface;
 			}
 
-			PreviewMeshComponent->SetMaterial( Index, RenderingMaterial ? RenderingMaterial : DisplayMaterialsMap[MaterialInterface].Get() );
+			PreviewMeshComponent->SetMaterial( Index, RenderingMaterial ? RenderingMaterial : MaterialInterface );
 		}
 
 		PreviewMeshComponent->MarkRenderStateDirty();
@@ -1253,14 +1127,6 @@ void SDataprepEditorViewport::ApplyRenderingMaterial()
 
 void SDataprepEditorViewport::LoadDefaultSettings()
 {
-	// Disable viewing settings for the time being
-	static bool bAllowViewingSettings = false;
-	if(!bAllowViewingSettings)
-	{
-		AssetViewerProfileIndex = 0;
-		return;
-	}
-
 	// Find index of Dataprep's viewport's settings
 	const TCHAR* DataprepViewportSettingProfileName = TEXT("DataprepViewportSetting");
 
@@ -1767,37 +1633,40 @@ namespace DataprepEditor3DPreviewUtils
 	// Copied from DatasmithImporterImpl::CompileMaterial
 	void CompileMaterials(const TArray< UMaterialInterface* > Materials)
 	{
-		FMaterialUpdateContext MaterialUpdateContext;
-
-		for(UMaterialInterface* MaterialInterface : Materials)
+		if(Materials.Num() > 0)
 		{
-			MaterialUpdateContext.AddMaterialInterface( MaterialInterface );
+			FMaterialUpdateContext MaterialUpdateContext;
 
-			if(UMaterialInstanceConstant* ConstantMaterialInstance = Cast< UMaterialInstanceConstant >(MaterialInterface))
+			for(UMaterialInterface* MaterialInterface : Materials)
 			{
-				// If BlendMode override property has been changed, make sure this combination of the parent material is compiled
-				if ( ConstantMaterialInstance->BasePropertyOverrides.bOverride_BlendMode == true )
-				{
-					ConstantMaterialInstance->ForceRecompileForRendering();
-				}
-				else
-				{
-					// If a static switch is overridden, we need to recompile
-					FStaticParameterSet StaticParameters;
-					ConstantMaterialInstance->GetStaticParameterValues( StaticParameters );
+				MaterialUpdateContext.AddMaterialInterface( MaterialInterface );
 
-					for ( FStaticSwitchParameter& Switch : StaticParameters.StaticSwitchParameters )
+				if(UMaterialInstanceConstant* ConstantMaterialInstance = Cast< UMaterialInstanceConstant >(MaterialInterface))
+				{
+					// If BlendMode override property has been changed, make sure this combination of the parent material is compiled
+					if ( ConstantMaterialInstance->BasePropertyOverrides.bOverride_BlendMode == true )
+					{
+						ConstantMaterialInstance->ForceRecompileForRendering();
+					}
+					else
 					{
-						if ( Switch.bOverride )
+						// If a static switch is overridden, we need to recompile
+						FStaticParameterSet StaticParameters;
+						ConstantMaterialInstance->GetStaticParameterValues( StaticParameters );
+
+						for ( FStaticSwitchParameter& Switch : StaticParameters.StaticSwitchParameters )
 						{
-							ConstantMaterialInstance->ForceRecompileForRendering();
-							break;
+							if ( Switch.bOverride )
+							{
+								ConstantMaterialInstance->ForceRecompileForRendering();
+								break;
+							}
 						}
 					}
-				}
 
-				ConstantMaterialInstance->PreEditChange( nullptr );
-				ConstantMaterialInstance->PostEditChange();
+					ConstantMaterialInstance->PreEditChange( nullptr );
+					ConstantMaterialInstance->PostEditChange();
+				}
 			}
 		}
 	}
@@ -1826,250 +1695,6 @@ namespace DataprepEditor3DPreviewUtils
 			}
 		}
 	}
-
-#ifdef VIEWPORT_EXPERIMENTAL
-	// Area using 3d positions
-	// Area is half length of the normal vector
-	float CalculateTriangleArea(const FVector& P0, const FVector& P1, const FVector& P2)
-	{
-		FVector Normal = ( P1 - P2 ) ^ ( P0 - P2 );
-		return Normal.Size() * 0.5f;
-	}
-
-	// Area using 2d positions
-	// Heron's formula: using length of the 3 sides
-	float CalculateTriangleArea(const FVector2D& P0, const FVector2D& P1, const FVector2D& P2)
-	{
-		return FMath::Sqrt( FMath::Abs( ( P1 - P2 ) ^ ( P0 - P2 ) ) ) * 0.5f;
-	}
-
-	// add exposure to colors to make them pop
-	void ExposureCompensation(FVector4& in)
-	{
-		for (int i = 0; i < 3; i++)
-		{
-			in[i] = FMath::Pow(in[i], 2.4);
-		}
-	}
-
-	void ComputeUVStretching(FStaticMeshLODResources& Resource)
-	{
-		if(!bComputeUVStretching)
-		{
-			return;
-		}
-
-		const FIndexArrayView& Indices = Resource.IndexBuffer.GetArrayView();
-		FStaticMeshVertexBuffer& VertexBuffer = Resource.VertexBuffers.StaticMeshVertexBuffer;
-		const FPositionVertexBuffer& Positions = Resource.VertexBuffers.PositionVertexBuffer;
-
-		const int32 UVIndex = 0;
-		const int32 NumTriangles = Indices.Num() / 3;
-		const int32 NumVertices = (int32)VertexBuffer.GetNumVertices();
-
-		TArray<FLinearColor> Colors;
-		Colors.AddZeroed( NumVertices );
-
-		float AverageArea3D = 0.0f;
-		float AverageArea2D = 0.0f;
-
-		for(int32 TriangleIndex = 0, Offset = 0; TriangleIndex < NumTriangles; ++TriangleIndex, Offset += 3)
-		{
-			const int32 Idx0 = Indices[Offset + 0];
-			const int32 Idx1 = Indices[Offset + 1];
-			const int32 Idx2 = Indices[Offset + 2];
-
-			// 2d
-			float Area2d = CalculateTriangleArea(
-				VertexBuffer.GetVertexUV_Typed<EStaticMeshVertexUVType::Default>(Idx0, UVIndex), 
-				VertexBuffer.GetVertexUV_Typed<EStaticMeshVertexUVType::Default>(Idx1, UVIndex),
-				VertexBuffer.GetVertexUV_Typed<EStaticMeshVertexUVType::Default>(Idx2, UVIndex)
-			);
-
-			// 3d
-			float Area3d = CalculateTriangleArea( Positions.VertexPosition(Idx0), Positions.VertexPosition(Idx1), Positions.VertexPosition(Idx2) );
-
-			{
-				FLinearColor& Color = Colors[Idx0];
-				Color.R += Area2d;
-				Color.G += Area3d;
-			}
-
-			{
-				FLinearColor& Color = Colors[Idx1];
-				Color.R += Area2d;
-				Color.G += Area3d;
-			}
-
-			{
-				FLinearColor& Color = Colors[Idx2];
-				Color.R += Area2d;
-				Color.G += Area3d;
-			}
-
-			// accumulate 
-			AverageArea2D += Area2d;
-			AverageArea3D += Area3d;
-		}
-
-		// average
-		AverageArea2D = AverageArea2D / NumTriangles;
-		AverageArea3D = AverageArea3D / NumTriangles;
-
-		//----------------------------------------------------------------------------------
-		// Step 2: Calculate distortion value per vert Instances
-		//		2d == 3d  -> optimal (white)
-		//		2d > 3d   -> compression (blue)
-		//		2d < 3d   -> stretching (red)
-		float DistortionRatioMin = 1.0f;
-		float DistortionRatioMax = 1.0f;
-
-		for(int32 VertexIndex = 0; VertexIndex < NumVertices; ++VertexIndex)
-		{
-			FLinearColor& Color = Colors[VertexIndex];
-
-			float Area2D = Color.R;
-			float Area3D = Color.G;
-			if(Area2D > SMALL_NUMBER && Area3D > SMALL_NUMBER )
-			{
-				Area2D /= AverageArea2D;
-				Area3D /= AverageArea3D;
-				const float DistortionRatio = Area2D / Area3D;
-
-				if(DistortionRatioMin > DistortionRatio)
-				{
-					DistortionRatioMin = DistortionRatio;
-				}
-
-				if(DistortionRatioMax < DistortionRatio)
-				{
-					DistortionRatioMax = DistortionRatio;
-				}
-
-				Color.R = Area2D;
-				Color.G = Area3D;
-				Color.A = DistortionRatio;
-			}
-		}
-
-		DistortionRatioMin = 1.0f / DistortionRatioMin;
-
-		float DistortionStrectchRatioRange = DistortionRatioMax - 1.0f;
-		float DistortionShrinkRatioRange = DistortionRatioMin - 1.0f;
-		for(int32 VertexIndex = 0; VertexIndex < NumVertices; ++VertexIndex)
-		{
-			FLinearColor& Color = Colors[VertexIndex];
-			const float DistortionRatio = Color.A;
-
-			if(DistortionRatio > SMALL_NUMBER)
-			{
-				FLinearColor DistortionColor(1, 1, 1, 1);
-				if(DistortionRatio < 1.0f)
-				{
-					const float D_Norm = ( (1.0f / DistortionRatio) - 1.0f) / DistortionShrinkRatioRange;
-					DistortionColor = FMath::Lerp(FLinearColor::Red, FLinearColor::White, D_Norm);
-				}
-				else
-				{
-					const float D_Norm = ( DistortionRatio - 1.0f) / DistortionStrectchRatioRange;
-					DistortionColor = FMath::Lerp(FLinearColor::Blue, FLinearColor::White, D_Norm);
-				}
-
-				Color = DistortionColor;
-			}
-		}
-
-		//----------------------------------------------------------------------------------
-		// Step 3 : Quadify colors 
-		// Average colors of every 2 triangles...
-		// Gives a faceted look
-		//int QuadItr = 0;
-		//TArray<int32> QuadIndexes;
-		//TArray<FVector4> OutQuadColors;
-		//OutQuadColors.AddZeroed(VertexInstancesCount);
-
-		//for (const FPolygonID PolygonID : MeshDescription.Polygons().GetElementIDs())
-		//{
-		//	const TArray<FMeshTriangle>& MeshTriangles = MeshDescription.GetPolygonTriangles(PolygonID);
-		//	for (const FMeshTriangle& MeshTriangle : MeshTriangles)
-		//	{
-		//		QuadIndexes.Add(MeshTriangle.GetVertexInstanceID(0).GetValue());
-		//		QuadIndexes.Add(MeshTriangle.GetVertexInstanceID(1).GetValue());
-		//		QuadIndexes.Add(MeshTriangle.GetVertexInstanceID(2).GetValue());
-		//	}
-
-		//	QuadItr++;
-		//	if (QuadItr % 2 == 0)
-		//	{
-		//		QuadItr = 0;
-		//		FVector4 AverageColor(0, 0, 0, 0);
-		//		for (int index : QuadIndexes)
-		//		{
-		//			AverageColor += OutColors[index];
-		//		}
-		//		AverageColor = AverageColor / QuadIndexes.Num();
-
-		//		for (int index : QuadIndexes)
-		//		{
-		//			OutQuadColors[index] = AverageColor;
-		//		}
-		//		QuadIndexes.Empty();
-		//	}
-		//}
-
-		//// last one
-		//if (QuadItr==1)
-		//{
-		//	for (int index : QuadIndexes)
-		//	{
-		//		OutQuadColors[index] = OutColors[index];
-		//	}
-		//}
-
-		//----------------------------------------------------------------------------------
-		// Smoother
-		//for (const FVertexID VertexID : MeshDescription.Vertices().GetElementIDs())
-		//{
-		//	FMeshVertex vert = MeshDescription.GetVertex(VertexID);
-		//	TArray<FVertexInstanceID> vertInstances = vert.VertexInstanceIDs;
-		//	FVector4 AverageColor(0, 0, 0, 0);
-
-		//	// accumulate
-		//	for (const FVertexInstanceID& VertexInstanceID : vertInstances)
-		//	{
-		//		AverageColor += OutQuadColors[VertexInstanceID.GetValue()];
-		//	}
-
-		//	// average
-		//	AverageColor = AverageColor / vertInstances.Num();
-
-		//	// save
-		//	for (const FVertexInstanceID& VertexInstanceID : vertInstances)
-		//	{
-		//		OutColors[VertexInstanceID.GetValue()] = AverageColor;
-		//	}
-		//}
-
-
-		//----------------------------------------------------------------------------------
-		// Exposure compensation
-		//for (FVector4& color : OutColors)
-		//{
-		//	ExposureCompensation(color);
-		//}
-		FColorVertexBuffer& VertexColors = Resource.VertexBuffers.ColorVertexBuffer;
-		for(int32 VertexIndex = 0; VertexIndex < NumVertices; ++VertexIndex)
-		{
-			FLinearColor& DistortionColor = Colors[VertexIndex];
-
-			//DistortionColor.R = FMath::Pow( DistortionColor.R, 2.4f );
-			//DistortionColor.G = FMath::Pow( DistortionColor.G, 2.4f );
-			//DistortionColor.B = FMath::Pow( DistortionColor.B, 2.4f );
-
-			VertexColors.VertexColor(VertexIndex) = DistortionColor.ToFColor(true);
-		}
-	}
-#endif
 }
 
 #undef LOCTEXT_NAMESPACE
diff --git a/Engine/Plugins/Experimental/DataPrepEditor/Source/DataPrepEditor/Private/Widgets/SDataprepEditorViewport.h b/Engine/Plugins/Experimental/DataPrepEditor/Source/DataPrepEditor/Private/Widgets/SDataprepEditorViewport.h
--- a/Engine/Plugins/Experimental/DataPrepEditor/Source/DataPrepEditor/Private/Widgets/SDataprepEditorViewport.h
+++ b/Engine/Plugins/Experimental/DataPrepEditor/Source/DataPrepEditor/Private/Widgets/SDataprepEditorViewport.h
@@ -162,8 +162,11 @@ public:
 	SDataprepEditorViewport();
 	~SDataprepEditorViewport();
 
-	void UpdateMeshes();
-	void ClearMeshes();
+	/** UpdateScene the viewport with the current content of the editor's preview world */
+	void UpdateScene();
+
+	/** Delete all the rendering data created by the viewport for display */
+	void ClearScene();
 
 	// ICommonEditorViewportToolbarInfoProvider interface
 	virtual TSharedRef<class SEditorViewport> GetViewportWidget() override;
@@ -252,9 +255,6 @@ private:
 	/** Initialize the materials used for the different rendering options */
 	void InitializeDefaultMaterials();
 
-	/** Initialize the set of materials to use for the mesh components to preview */
-	void CreateDisplayMaterials(const TArray< UStaticMeshComponent* >& MeshComponentsToPreview);
-
 	/** Set the material to apply on the preview mesh components */
 	void SetRenderingMaterial(ERenderingMaterialType InRenderingMaterial);
 
@@ -326,12 +326,6 @@ private:
 	/** Mapping between preview's static mesh component and scene's one */
 	TMap< UStaticMeshComponent*, UStaticMeshComponent* > MeshComponentsReverseMapping;
 
-	/** Array of static mesh components which render state needs to be restored */
-	TArray<  TWeakObjectPtr< UStaticMeshComponent > > MeshComponentsToRestore;
-
-	/** Mapping between scene's materials and preview's materials */
-	TMap< class UMaterialInterface*, TWeakObjectPtr<UMaterialInstanceConstant> > DisplayMaterialsMap;
-
 	/** Array of selected static mesh components */
 	TSet<UStaticMeshComponent*> SelectedPreviewComponents;
 
@@ -367,9 +361,6 @@ private:
 	static TWeakObjectPtr<UMaterial> PerMeshMaterial;
 	static TArray<TWeakObjectPtr<UMaterialInstanceConstant>> PerMeshMaterialInstances;
 
-	/** Array of meshes built for viewing */
-	TArray<UStaticMesh*> BuiltMeshes;
-
 	TWeakPtr<FDataprepEditor> DataprepEditor;
 
 	/** Pointer to the vertical box into which the overlay text items are added */
diff --git a/Engine/Plugins/Experimental/DataPrepEditor/Source/DataPrepLibraries/Private/DataPrepOperationsLibrary.cpp b/Engine/Plugins/Experimental/DataPrepEditor/Source/DataPrepLibraries/Private/DataPrepOperationsLibrary.cpp
--- a/Engine/Plugins/Experimental/DataPrepEditor/Source/DataPrepLibraries/Private/DataPrepOperationsLibrary.cpp
+++ b/Engine/Plugins/Experimental/DataPrepEditor/Source/DataPrepLibraries/Private/DataPrepOperationsLibrary.cpp
@@ -38,7 +38,7 @@ extern UNREALED_API UEditorEngine* GEditor;
 
 namespace DataprepOperationsLibraryUtil
 {
-	TArray<UStaticMesh*> GetSelectedMeshes(const TArray<AActor*>& SelectedActors)
+	TSet<UStaticMesh*> GetSelectedMeshes(const TArray<AActor*>& SelectedActors)
 	{
 		TSet<UStaticMesh*> SelectedMeshes;
 
@@ -49,32 +49,34 @@ namespace DataprepOperationsLibraryUtil
 				TInlineComponentArray<UStaticMeshComponent*> StaticMeshComponents(Actor);
 				for (UStaticMeshComponent* StaticMeshComponent : StaticMeshComponents)
 				{
-					SelectedMeshes.Add(StaticMeshComponent->GetStaticMesh());
+					if(UStaticMesh* StaticMesh = StaticMeshComponent->GetStaticMesh())
+					{
+						SelectedMeshes.Add( StaticMesh );
+					}
 				}
 			}
 		}
 
-		return SelectedMeshes.Array();
+		return SelectedMeshes;
 	}
 
-	TArray<UStaticMesh*> GetSelectedMeshes(const TArray<UObject*>& SelectedObjects)
+	TSet<UStaticMesh*> GetSelectedMeshes(const TArray<UObject*>& SelectedObjects)
 	{
 		TSet<UStaticMesh*> SelectedMeshes;
 		TArray<AActor*> SelectedActors;
 
-		// Create LODs but do not commit changes
 		for (UObject* Object : SelectedObjects)
 		{
-			if (Object->IsA(UStaticMesh::StaticClass()))
+			if ( UStaticMesh* StaticMesh = Cast<UStaticMesh>(Object) )
 			{
-				SelectedMeshes.Add(Cast<UStaticMesh>(Object));
+				SelectedMeshes.Add( StaticMesh );
 			}
-			else if (Object->IsA(AActor::StaticClass()))
+			else if (AActor* Actor = Cast<AActor>(Object) )
 			{
-				TInlineComponentArray<UStaticMeshComponent*> StaticMeshComponents(Cast<AActor>(Object));
+				TInlineComponentArray<UStaticMeshComponent*> StaticMeshComponents( Actor );
 				for (UStaticMeshComponent* StaticMeshComponent : StaticMeshComponents)
 				{
-					if(UStaticMesh* StaticMesh = StaticMeshComponent->GetStaticMesh())
+					if((StaticMesh = StaticMeshComponent->GetStaticMesh()) != nullptr)
 					{
 						SelectedMeshes.Add( StaticMesh );
 					}
@@ -82,7 +84,7 @@ namespace DataprepOperationsLibraryUtil
 			}
 		}
 
-		return SelectedMeshes.Array();
+		return SelectedMeshes;
 	}
 
 	TArray<UMaterialInterface*> GetUsedMaterials(const TArray<UObject*>& SelectedObjects)
@@ -244,54 +246,132 @@ namespace DataprepOperationsLibraryUtil
 		}
 	}
 
-	// Replacement of UStaticMesh::CacheDerivedData() which performs too much operations for our purpose
-	// And displays unwanted progress bar
-	// #ueent_todo: Work with Geometry team to find the proper replacement
-	void BuildRenderData( UStaticMesh* StaticMesh, IMeshBuilderModule& MeshBuilderModule )
+	FStaticMeshBuilder::FStaticMeshBuilder(const TSet<UStaticMesh *>& InStaticMeshes)
 	{
-		if(StaticMesh->GetNumSourceModels() == 0)
+		StaticMeshes = BuildStaticMeshes( InStaticMeshes );
+	}
+
+	FStaticMeshBuilder::~FStaticMeshBuilder()
+	{
+		// Release render data of built static meshes
+		for(UStaticMesh* StaticMesh : StaticMeshes)
 		{
-			return;
+			if(StaticMesh)
+			{
+				StaticMesh->RenderData.Reset();
+			}
 		}
+	}
 
-		TRACE_CPUPROFILER_EVENT_SCOPE(DataprepOperationsLibraryUtil::BuildRenderData);
+	TArray<UStaticMesh*> BuildStaticMeshes(const TSet<UStaticMesh*>& StaticMeshes, bool bForceBuild)
+	{
+		TRACE_CPUPROFILER_EVENT_SCOPE(DataprepOperationsLibraryUtil::BuildStaticMeshes);
 
-		FStaticMeshSourceModel& SourceModel = StaticMesh->GetSourceModel(0);
+		TArray<UStaticMesh*> BuiltMeshes;
+		BuiltMeshes.Reserve( StaticMeshes.Num() );
 
-		FMeshBuildSettings PrevBuildSettings = SourceModel.BuildSettings;
-		SourceModel.BuildSettings.bGenerateLightmapUVs = false;
-		SourceModel.BuildSettings.bRecomputeNormals = false;
-		SourceModel.BuildSettings.bRecomputeTangents = false;
-		SourceModel.BuildSettings.bBuildAdjacencyBuffer = false;
-		SourceModel.BuildSettings.bBuildReversedIndexBuffer = false;
-		SourceModel.BuildSettings.bComputeWeightedNormals = false;
+		if(bForceBuild)
+		{
+			BuiltMeshes.Append( StaticMeshes.Array() );
+		}
+		else
+		{
+			for(UStaticMesh* StaticMesh : StaticMeshes)
+			{
+				if(StaticMesh && (!StaticMesh->RenderData.IsValid() || !StaticMesh->RenderData->IsInitialized()))
+				{
+					BuiltMeshes.Add( StaticMesh );
+				}
+			}
+		}
 
-		FMeshDescription* MeshDescription = SourceModel.MeshDescription.Get();
-		check( MeshDescription );
+		if(BuiltMeshes.Num() > 0)
+		{
+			// Start with the biggest mesh first to help balancing tasks on threads
+			BuiltMeshes.Sort(
+				[](const UStaticMesh& Lhs, const UStaticMesh& Rhs) 
+			{ 
+				int32 LhsVerticesNum = Lhs.IsMeshDescriptionValid(0) ? Lhs.GetMeshDescription(0)->Vertices().Num() : 0;
+				int32 RhsVerticesNum = Rhs.IsMeshDescriptionValid(0) ? Rhs.GetMeshDescription(0)->Vertices().Num() : 0;
+
+				return LhsVerticesNum > RhsVerticesNum;
+			}
+			);
 
-		// Create render data
-		StaticMesh->RenderData.Reset( new(FMemory::Malloc(sizeof(FStaticMeshRenderData)))FStaticMeshRenderData() );
+			//Cache the BuildSettings and update them before building the meshes.
+			TArray< TArray<FMeshBuildSettings> > StaticMeshesSettings;
+			StaticMeshesSettings.Reserve( BuiltMeshes.Num() );
 
-		ITargetPlatformManagerModule& TargetPlatformManager = GetTargetPlatformManagerRef();
-		ITargetPlatform* RunningPlatform = TargetPlatformManager.GetRunningTargetPlatform();
-		check(RunningPlatform);
-		const FStaticMeshLODSettings& LODSettings = RunningPlatform->GetStaticMeshLODSettings();
+			for (UStaticMesh* StaticMesh : BuiltMeshes)
+			{
+				TArray<FStaticMeshSourceModel>& SourceModels = StaticMesh->GetSourceModels();
+				TArray<FMeshBuildSettings> BuildSettings;
+				BuildSettings.Reserve(SourceModels.Num());
 
-		const FStaticMeshLODGroup& LODGroup = LODSettings.GetLODGroup(StaticMesh->LODGroup);
+				for(int32 Index = 0; Index < SourceModels.Num(); ++Index)
+				{
+					FStaticMeshSourceModel& SourceModel = SourceModels[Index];
 
-		if ( !MeshBuilderModule.BuildMesh( *StaticMesh->RenderData, StaticMesh, LODGroup ) )
-		{
-			UE_LOG(LogDataprep, Error, TEXT("Failed to build static mesh. See previous line(s) for details."));
-			return;
+					BuildSettings.Add( SourceModel.BuildSettings );
+
+					if(FMeshDescription* MeshDescription = StaticMesh->GetMeshDescription(Index))
+					{
+						FStaticMeshAttributes Attributes(*MeshDescription);
+						if(SourceModel.BuildSettings.DstLightmapIndex != -1)
+						{
+							TVertexInstanceAttributesConstRef<FVector2D> VertexInstanceUVs = Attributes.GetVertexInstanceUVs();
+							SourceModel.BuildSettings.bGenerateLightmapUVs = VertexInstanceUVs.IsValid() && VertexInstanceUVs.GetNumIndices() > SourceModel.BuildSettings.DstLightmapIndex;
+						}
+						else
+						{
+							SourceModel.BuildSettings.bGenerateLightmapUVs = false;
+						}
+
+						SourceModel.BuildSettings.bRecomputeNormals = !(Attributes.GetVertexInstanceNormals().IsValid() && Attributes.GetVertexInstanceNormals().GetNumIndices() > 0);
+						SourceModel.BuildSettings.bRecomputeTangents = false;
+						//SourceModel.BuildSettings.bBuildAdjacencyBuffer = false;
+						//SourceModel.BuildSettings.bBuildReversedIndexBuffer = false;
+					}
+				}
+
+				StaticMeshesSettings.Add(MoveTemp(BuildSettings));				
+			}
+
+			// Disable warnings from LogStaticMesh. Not useful
+			ELogVerbosity::Type PrevLogStaticMeshVerbosity = LogStaticMesh.GetVerbosity();
+			LogStaticMesh.SetVerbosity( ELogVerbosity::Error );
+
+			UStaticMesh::BatchBuild(BuiltMeshes, true );
+
+			// Restore LogStaticMesh verbosity
+			LogStaticMesh.SetVerbosity( PrevLogStaticMeshVerbosity );
+
+			for(int32 Index = 0; Index < BuiltMeshes.Num(); ++Index)
+			{
+				UStaticMesh* StaticMesh = BuiltMeshes[Index];
+				TArray<FMeshBuildSettings>& PrevBuildSettings = StaticMeshesSettings[Index];
+
+				TArray<FStaticMeshSourceModel>& SourceModels = StaticMesh->GetSourceModels();
+
+				for(int32 SourceModelIndex = 0; SourceModelIndex < SourceModels.Num(); ++SourceModelIndex)
+				{
+					SourceModels[SourceModelIndex].BuildSettings = PrevBuildSettings[SourceModelIndex];
+				}
+
+				for ( FStaticMeshLODResources& LODResources : StaticMesh->RenderData->LODResources )
+				{
+					LODResources.bHasColorVertexData = true;
+				}
+			}
 		}
 
-		SourceModel.BuildSettings = PrevBuildSettings;
+		return BuiltMeshes;
 	}
 } // ns DataprepOperationsLibraryUtil
 
-void UDataprepOperationsLibrary::SetLods(const TArray<UObject*>& SelectedObjects, const FEditorScriptingMeshReductionOptions& ReductionOptions)
+void UDataprepOperationsLibrary::SetLods(const TArray<UObject*>& SelectedObjects, const FEditorScriptingMeshReductionOptions& ReductionOptions, TArray<UObject*>& ModifiedObjects)
 {
-	TArray<UStaticMesh*> SelectedMeshes = DataprepOperationsLibraryUtil::GetSelectedMeshes(SelectedObjects);
+	TSet<UStaticMesh*> SelectedMeshes = DataprepOperationsLibraryUtil::GetSelectedMeshes(SelectedObjects);
 
 	// Create LODs but do not commit changes
 	for (UStaticMesh* StaticMesh : SelectedMeshes)
@@ -301,15 +381,36 @@ void UDataprepOperationsLibrary::SetLods(const TArray<UObject*>& SelectedObjects
 			DataprepOperationsLibraryUtil::FScopedStaticMeshEdit StaticMeshEdit( StaticMesh );
 
 			UEditorStaticMeshLibrary::SetLodsWithNotification(StaticMesh, ReductionOptions, false);
+
+			ModifiedObjects.Add( StaticMesh );
 		}
 	}
 }
 
-void UDataprepOperationsLibrary::SetSimpleCollision(const TArray<UObject*>& SelectedObjects, const EScriptingCollisionShapeType ShapeType)
+void UDataprepOperationsLibrary::SetSimpleCollision(const TArray<UObject*>& SelectedObjects, const EScriptingCollisionShapeType ShapeType, TArray<UObject*>& ModifiedObjects)
 {
-	TArray<UStaticMesh*> SelectedMeshes = DataprepOperationsLibraryUtil::GetSelectedMeshes(SelectedObjects);
+	TSet<UStaticMesh*> SelectedMeshes = DataprepOperationsLibraryUtil::GetSelectedMeshes(SelectedObjects);
+
+	// Make sure all static meshes to be processed have render data for NDOP types
+	bool bNeedRenderData = false;
+	switch (ShapeType)
+	{
+		case EScriptingCollisionShapeType::NDOP10_X:
+		case EScriptingCollisionShapeType::NDOP10_Y:
+		case EScriptingCollisionShapeType::NDOP10_Z:
+		case EScriptingCollisionShapeType::NDOP18:
+		case EScriptingCollisionShapeType::NDOP26:
+		{
+			bNeedRenderData = true;
+			break;
+		}
+		default:
+		{
+			break;
+		}
+	}
 
-	IMeshBuilderModule& MeshBuilderModule = FModuleManager::LoadModuleChecked< IMeshBuilderModule >( TEXT("MeshBuilder") );
+	DataprepOperationsLibraryUtil::FStaticMeshBuilder StaticMeshBuilder( bNeedRenderData ? SelectedMeshes : TSet<UStaticMesh*>() );
 
 	// Create LODs but do not commit changes
 	for (UStaticMesh* StaticMesh : SelectedMeshes)
@@ -321,42 +422,19 @@ void UDataprepOperationsLibrary::SetSimpleCollision(const TArray<UObject*>& Sele
 			// Remove existing simple collisions
 			UEditorStaticMeshLibrary::RemoveCollisionsWithNotification( StaticMesh, false );
 
-			// Check that render data is available if k-DOP type of collision is required
-			bool bFreeRenderData = false;
-			switch (ShapeType)
-			{
-				case EScriptingCollisionShapeType::NDOP10_X:
-				case EScriptingCollisionShapeType::NDOP10_Y:
-				case EScriptingCollisionShapeType::NDOP10_Z:
-				case EScriptingCollisionShapeType::NDOP18:
-				case EScriptingCollisionShapeType::NDOP26:
-				{
-					if( StaticMesh->RenderData == nullptr )
-					{
-						DataprepOperationsLibraryUtil::BuildRenderData( StaticMesh, MeshBuilderModule );
-						bFreeRenderData = true;
-					}
-					break;
-				}
-				default:
-				{
-					break;
-				}
-			}
-
 			UEditorStaticMeshLibrary::AddSimpleCollisionsWithNotification( StaticMesh, ShapeType, false );
 
-			if( bFreeRenderData )
-			{
-				StaticMesh->RenderData = nullptr;
-			}
+			ModifiedObjects.Add( StaticMesh );
 		}
 	}
 }
 
-void UDataprepOperationsLibrary::SetConvexDecompositionCollision(const TArray<UObject*>& SelectedObjects, int32 HullCount, int32 MaxHullVerts, int32 HullPrecision)
+void UDataprepOperationsLibrary::SetConvexDecompositionCollision(const TArray<UObject*>& SelectedObjects, int32 HullCount, int32 MaxHullVerts, int32 HullPrecision, TArray<UObject*>& ModifiedObjects)
 {
-	TArray<UStaticMesh*> SelectedMeshes = DataprepOperationsLibraryUtil::GetSelectedMeshes(SelectedObjects);
+	TSet<UStaticMesh*> SelectedMeshes = DataprepOperationsLibraryUtil::GetSelectedMeshes(SelectedObjects);
+
+	// Make sure all static meshes to be processed have render data
+	DataprepOperationsLibraryUtil::FStaticMeshBuilder StaticMeshBuilder(SelectedMeshes);
 
 	// Build complex collision
 	for (UStaticMesh* StaticMesh : SelectedMeshes)
@@ -366,13 +444,15 @@ void UDataprepOperationsLibrary::SetConvexDecompositionCollision(const TArray<UO
 			DataprepOperationsLibraryUtil::FScopedStaticMeshEdit StaticMeshEdit( StaticMesh );
 
 			UEditorStaticMeshLibrary::SetConvexDecompositionCollisionsWithNotification(StaticMesh, HullCount, MaxHullVerts, HullPrecision, false);
+
+			ModifiedObjects.Add( StaticMesh );
 		}
 	}
 }
 
-void UDataprepOperationsLibrary::SetGenerateLightmapUVs( const TArray< UObject* >& Assets, bool bGenerateLightmapUVs )
+void UDataprepOperationsLibrary::SetGenerateLightmapUVs( const TArray< UObject* >& Assets, bool bGenerateLightmapUVs, TArray<UObject*>& ModifiedObjects )
 {
-	TArray<UStaticMesh*> SelectedMeshes = DataprepOperationsLibraryUtil::GetSelectedMeshes(Assets);
+	TSet<UStaticMesh*> SelectedMeshes = DataprepOperationsLibraryUtil::GetSelectedMeshes(Assets);
 
 	for (UStaticMesh* StaticMesh : SelectedMeshes)
 	{
@@ -399,6 +479,8 @@ void UDataprepOperationsLibrary::SetGenerateLightmapUVs( const TArray< UObject*
 				// Correct the coordinate index if it was invalid
 				StaticMesh->LightMapCoordinateIndex = MinBiggestUVChannel;
 			}
+
+			ModifiedObjects.Add( StaticMesh );
 		}
 	}
 }
@@ -527,19 +609,23 @@ void UDataprepOperationsLibrary::SetMaterial( const TArray< UObject* >& Selected
 	}
 }
 
-void UDataprepOperationsLibrary::SetLODGroup( const TArray<UObject*>& SelectedObjects, FName& LODGroupName )
+void UDataprepOperationsLibrary::SetLODGroup( const TArray<UObject*>& SelectedObjects, FName& LODGroupName, TArray<UObject*>& ModifiedObjects )
 {
 	TArray<FName> LODGroupNames;
 	UStaticMesh::GetLODGroups( LODGroupNames );
 
 	if ( LODGroupNames.Find( LODGroupName ) != INDEX_NONE )
 	{
-		TArray<UStaticMesh*> SelectedMeshes = DataprepOperationsLibraryUtil::GetSelectedMeshes(SelectedObjects);
+		TSet<UStaticMesh*> SelectedMeshes = DataprepOperationsLibraryUtil::GetSelectedMeshes(SelectedObjects);
 
 		// Apply the new LODGroup without rebuilding the static mesh
 		for (UStaticMesh* StaticMesh : SelectedMeshes)
 		{
-			StaticMesh->SetLODGroup( LODGroupName, false);
+			if(StaticMesh)
+			{
+				StaticMesh->SetLODGroup( LODGroupName, false);
+				ModifiedObjects.Add( StaticMesh );
+			}
 		}
 	}
 }
diff --git a/Engine/Plugins/Experimental/DataPrepEditor/Source/DataPrepLibraries/Private/DataPrepOperationsLibrary.h b/Engine/Plugins/Experimental/DataPrepEditor/Source/DataPrepLibraries/Private/DataPrepOperationsLibrary.h
--- a/Engine/Plugins/Experimental/DataPrepEditor/Source/DataPrepLibraries/Private/DataPrepOperationsLibrary.h
+++ b/Engine/Plugins/Experimental/DataPrepEditor/Source/DataPrepLibraries/Private/DataPrepOperationsLibrary.h
@@ -21,7 +21,29 @@ DECLARE_LOG_CATEGORY_EXTERN( LogDataprep, Log, All );
 
 namespace DataprepOperationsLibraryUtil
 {
-	void BuildRenderData( UStaticMesh* StaticMesh, IMeshBuilderModule& MeshBuilderModule );
+	class FStaticMeshBuilder
+	{
+	public:
+		FStaticMeshBuilder(const TSet<UStaticMesh*>& InStaticMeshes);
+		~FStaticMeshBuilder();
+	private:
+		TArray<UStaticMesh*> StaticMeshes;
+	};
+
+	/*
+	 * Builds render data of a set of static meshes.
+	 * @param StaticMeshes	Set of static meshes to build if render data is missing or a forced build is required
+	 * @param bForceBuild	Indicates if all static meshes should be built or only the incomplete ones
+	 * @returns the array of static meshes which have actually been built
+	 */
+	TArray<UStaticMesh*> BuildStaticMeshes( const TSet<UStaticMesh*>& StaticMeshes, bool bForceBuild = false );
+
+	/*
+	 * Find the set of static meshes in or referenced by a given array of objects.
+	 * @param SelectedObjects	Array of UObjects to go through
+	 * @returns a set of static meshes
+	 */
+	TSet<UStaticMesh*> GetSelectedMeshes(const TArray<UObject*>& SelectedObjects);
 }
 
 /*
@@ -134,7 +156,7 @@ public:
 	 * Calls UEditorStaticMeshLibrary::SetLods on each static mesh of the resulting array.
 	 */
 	UFUNCTION(BlueprintCallable, Category = "Dataprep | Operation")
-	static void SetLods(const TArray<UObject*>& SelectedObjects, const FEditorScriptingMeshReductionOptions& ReductionOptions);
+	static void SetLods(const TArray<UObject*>& SelectedObjects, const FEditorScriptingMeshReductionOptions& ReductionOptions, TArray<UObject*>& ModifiedObjects);
 
 	/**
 	 * Set one simple collision of the given shape type on the static meshes contained in the
@@ -150,7 +172,7 @@ public:
 	 * Calls UEditorStaticMeshLibrary::AddSimpleCollisions on each static mesh of the resulting array.
 	 */
 	UFUNCTION(BlueprintCallable, Category = "Dataprep | Operation")
-	static void SetSimpleCollision(const TArray<UObject*>& SelectedObjects, const EScriptingCollisionShapeType ShapeType);
+	static void SetSimpleCollision(const TArray<UObject*>& SelectedObjects, const EScriptingCollisionShapeType ShapeType, TArray<UObject*>& ModifiedObjects);
 
 	/**
 	 * Add complex collision on the static meshes contained in the input array
@@ -167,7 +189,7 @@ public:
 	 * Note that any simple collisions on each static mesh of the resulting array will be removed.
 	 */
 	UFUNCTION(BlueprintCallable, Category = "Dataprep | Operation")
-	static void SetConvexDecompositionCollision(const TArray<UObject*>& SelectedObjects, int32 HullCount, int32 MaxHullVerts, int32 HullPrecision);
+	static void SetConvexDecompositionCollision(const TArray<UObject*>& SelectedObjects, int32 HullCount, int32 MaxHullVerts, int32 HullPrecision, TArray<UObject*>& ModifiedObjects);
 
 	/**
 	 * Sets the Generate Lightmap UVs flag on the static meshes found in the Assets list
@@ -176,7 +198,7 @@ public:
 	 * @param	bGenerateLightmapUVs	The value to set for the generate lightmap uvs flag.
 	 */
 	UFUNCTION(BlueprintCallable, Category = "Dataprep | Operation")
-	static void SetGenerateLightmapUVs( const TArray< UObject* >& Assets, bool bGenerateLightmapUVs );
+	static void SetGenerateLightmapUVs( const TArray< UObject* >& Assets, bool bGenerateLightmapUVs, TArray<UObject*>& ModifiedObjects );
 
 	/**
 	 * Replaces designated materials in all or specific content folders with specific ones
@@ -205,7 +227,7 @@ public:
 	 * @remark: Static meshes are not re-built after the new LOD groups are set
 	 */
 	UFUNCTION(BlueprintCallable, Category = "Dataprep | Operation")
-	static void SetLODGroup( const TArray< UObject* >& SelectedObjects, FName& LODGroupName );
+	static void SetLODGroup( const TArray< UObject* >& SelectedObjects, FName& LODGroupName, TArray<UObject*>& ModifiedObjects );
 
 	/**
 	 * Set the material to all elements of a set of Static Meshes or Static Mesh Actors
diff --git a/Engine/Plugins/Experimental/DataPrepEditor/Source/DataPrepLibraries/Private/DataprepEditingOperations.cpp b/Engine/Plugins/Experimental/DataPrepEditor/Source/DataPrepLibraries/Private/DataprepEditingOperations.cpp
--- a/Engine/Plugins/Experimental/DataPrepEditor/Source/DataPrepLibraries/Private/DataprepEditingOperations.cpp
+++ b/Engine/Plugins/Experimental/DataPrepEditor/Source/DataPrepLibraries/Private/DataprepEditingOperations.cpp
@@ -278,7 +278,7 @@ bool UDataprepMergeActorsOperation::MergeStaticMeshActors(UWorld* World, const T
 		}
 	}
 
-	FStaticMeshBuilder StaticMeshBuilder(StaticMeshes);
+	DataprepOperationsLibraryUtil::FStaticMeshBuilder StaticMeshBuilder( StaticMeshes );
 
 	//
 	// See MeshMergingTool.cpp
@@ -300,7 +300,7 @@ bool UDataprepMergeActorsOperation::MergeStaticMeshActors(UWorld* World, const T
 	}
 
 	// Add asset to set of assets in Dataprep action working set
-	MergedMesh = Cast<UStaticMesh>( AddAsset( UtilitiesMergedMesh, UStaticMesh::StaticClass(), NewActorLabel.IsEmpty() ? TEXT("Merged_Mesh") : *NewActorLabel ) );
+	MergedMesh = Cast<UStaticMesh>( AddAsset( UtilitiesMergedMesh, NewActorLabel.IsEmpty() ? TEXT("Merged_Mesh") : *NewActorLabel ) );
 	if (!MergedMesh)
 	{
 		UE_LOG(LogDataprep, Error, TEXT("MergeStaticMeshActors failed. Internal error while creating the merged mesh."));
@@ -325,46 +325,6 @@ bool UDataprepMergeActorsOperation::MergeStaticMeshActors(UWorld* World, const T
 	return true;
 }
 
-UDataprepMergeActorsOperation::FStaticMeshBuilder::FStaticMeshBuilder(const TSet<UStaticMesh *>& InStaticMeshes)
-{
-	StaticMeshes = InStaticMeshes;
-
-	IMeshBuilderModule& MeshBuilderModule = FModuleManager::LoadModuleChecked< IMeshBuilderModule >( TEXT("MeshBuilder") );
-	UDataprepMergeActorsOperation::PrepareStaticMeshes(StaticMeshes, MeshBuilderModule);
-}
-
-UDataprepMergeActorsOperation::FStaticMeshBuilder::~FStaticMeshBuilder()
-{
-	// Release render data of built static meshes
-	for(UStaticMesh* StaticMesh : StaticMeshes)
-	{
-		if(StaticMesh)
-		{
-			StaticMesh->RenderData.Reset();
-		}
-	}
-}
-
-void UDataprepMergeActorsOperation::PrepareStaticMeshes(const TSet<UStaticMesh*>& StaticMeshes, IMeshBuilderModule& MeshBuilderModule)
-{
-	TRACE_CPUPROFILER_EVENT_SCOPE(UDataprepMergeActorsOperation::PrepareStaticMeshes);
-
-	if( StaticMeshes.Num() > 1 )
-	{
-		TArray<UStaticMesh*> StaticMeshesToBuild = StaticMeshes.Array();
-		ParallelFor( StaticMeshesToBuild.Num(), [&]( int32 Index ) {
-			DataprepOperationsLibraryUtil::BuildRenderData( StaticMeshesToBuild[Index], MeshBuilderModule );
-		});
-	}
-	else
-	{
-		for(UStaticMesh* StaticMesh : StaticMeshes)
-		{
-			DataprepOperationsLibraryUtil::BuildRenderData( StaticMesh, MeshBuilderModule );
-		}
-	}
-}
-
 void UDataprepDeleteUnusedAssetsOperation::OnExecution_Implementation(const FDataprepContext& InContext)
 {
 	UWorld* World = nullptr;
diff --git a/Engine/Plugins/Experimental/DataPrepEditor/Source/DataPrepLibraries/Private/DataprepEditingOperations.h b/Engine/Plugins/Experimental/DataPrepEditor/Source/DataPrepLibraries/Private/DataprepEditingOperations.h
--- a/Engine/Plugins/Experimental/DataPrepEditor/Source/DataPrepLibraries/Private/DataprepEditingOperations.h
+++ b/Engine/Plugins/Experimental/DataPrepEditor/Source/DataPrepLibraries/Private/DataprepEditingOperations.h
@@ -76,22 +76,11 @@ public:
 	bool bPivotPointAtZero;
 
 protected:
-	class FStaticMeshBuilder
-	{
-	public:
-		FStaticMeshBuilder(const TSet<UStaticMesh*>& InStaticMeshes);
-		~FStaticMeshBuilder();
-	private:
-		TSet<UStaticMesh*> StaticMeshes;
-	};
-
 	virtual void OnExecution_Implementation(const FDataprepContext& InContext) override;
 	//~ End UDataprepOperation Interface
 
 	bool MergeStaticMeshActors(UWorld* World, const TArray<UPrimitiveComponent*>& ComponentsToMerge, const FString& RootName, bool bCreateActor = true);
 
-	static void PrepareStaticMeshes( const TSet<UStaticMesh*>& StaticMeshes, IMeshBuilderModule& MeshBuilderModule );
-
 protected:
 	FVector MergedMeshWorldLocation;
 	UStaticMesh* MergedMesh;
diff --git a/Engine/Plugins/Experimental/DataPrepEditor/Source/DataPrepLibraries/Private/DataprepOperations.cpp b/Engine/Plugins/Experimental/DataPrepEditor/Source/DataPrepLibraries/Private/DataprepOperations.cpp
--- a/Engine/Plugins/Experimental/DataPrepEditor/Source/DataPrepLibraries/Private/DataprepOperations.cpp
+++ b/Engine/Plugins/Experimental/DataPrepEditor/Source/DataPrepLibraries/Private/DataprepOperations.cpp
@@ -82,7 +82,13 @@ void UDataprepSetLODsOperation::OnExecution_Implementation(const FDataprepContex
 #endif
 
 	// Execute operation
-	UDataprepOperationsLibrary::SetLods( InContext.Objects, ReductionOptions );
+	TArray<UObject*> ModifiedStaticMeshes;
+	UDataprepOperationsLibrary::SetLods( InContext.Objects, ReductionOptions, ModifiedStaticMeshes );
+
+	if(ModifiedStaticMeshes.Num() > 0)
+	{
+		AssetsModified( MoveTemp( ModifiedStaticMeshes ) );
+	}
 }
 
 UDataprepSetLODGroupOperation::UDataprepSetLODGroupOperation()
@@ -100,7 +106,13 @@ void UDataprepSetLODGroupOperation::OnExecution_Implementation(const FDataprepCo
 #endif
 
 	// Execute operation
-	UDataprepOperationsLibrary::SetLODGroup( InContext.Objects, GroupName );
+	TArray<UObject*> ModifiedStaticMeshes;
+	UDataprepOperationsLibrary::SetLODGroup( InContext.Objects, GroupName, ModifiedStaticMeshes );
+
+	if(ModifiedStaticMeshes.Num() > 0)
+	{
+		AssetsModified( MoveTemp( ModifiedStaticMeshes ) );
+	}
 }
 
 void UDataprepSetSimpleCollisionOperation::OnExecution_Implementation(const FDataprepContext& InContext)
@@ -110,7 +122,13 @@ void UDataprepSetSimpleCollisionOperation::OnExecution_Implementation(const FDat
 #endif
 
 	// Execute operation
-	UDataprepOperationsLibrary::SetSimpleCollision( InContext.Objects, ShapeType );
+	TArray<UObject*> ModifiedStaticMeshes;
+	UDataprepOperationsLibrary::SetSimpleCollision( InContext.Objects, ShapeType, ModifiedStaticMeshes );
+
+	if(ModifiedStaticMeshes.Num() > 0)
+	{
+		AssetsModified( MoveTemp( ModifiedStaticMeshes ) );
+	}
 }
 
 void UDataprepSetConvexDecompositionCollisionOperation::OnExecution_Implementation(const FDataprepContext& InContext)
@@ -120,7 +138,13 @@ void UDataprepSetConvexDecompositionCollisionOperation::OnExecution_Implementati
 #endif
 
 	// Execute operation
-	UDataprepOperationsLibrary::SetConvexDecompositionCollision( InContext.Objects, HullCount, MaxHullVerts, HullPrecision );
+	TArray<UObject*> ModifiedStaticMeshes;
+	UDataprepOperationsLibrary::SetConvexDecompositionCollision( InContext.Objects, HullCount, MaxHullVerts, HullPrecision, ModifiedStaticMeshes );
+
+	if(ModifiedStaticMeshes.Num() > 0)
+	{
+		AssetsModified( MoveTemp( ModifiedStaticMeshes ) );
+	}
 }
 
 void UDataprepSetGenerateLightmapUVsOperation::OnExecution_Implementation(const FDataprepContext& InContext)
@@ -130,7 +154,13 @@ void UDataprepSetGenerateLightmapUVsOperation::OnExecution_Implementation(const
 #endif
 
 	// Execute operation
-	UDataprepOperationsLibrary::SetGenerateLightmapUVs( InContext.Objects, bGenerateLightmapUVs );
+	TArray<UObject*> ModifiedStaticMeshes;
+	UDataprepOperationsLibrary::SetGenerateLightmapUVs( InContext.Objects, bGenerateLightmapUVs, ModifiedStaticMeshes );
+
+	if(ModifiedStaticMeshes.Num() > 0)
+	{
+		AssetsModified( MoveTemp( ModifiedStaticMeshes ) );
+	}
 }
 
 void UDataprepSetMobilityOperation::OnExecution_Implementation(const FDataprepContext& InContext)
