Fixed bad merge from //UE4/Dev-Enterprise

diff --git a/Engine/Plugins/Experimental/DataPrepEditor/Config/BaseDataprepEditor.ini b/Engine/Plugins/Experimental/DataPrepEditor/Config/BaseDataprepEditor.ini
--- a/Engine/Plugins/Experimental/DataPrepEditor/Config/BaseDataprepEditor.ini
+++ b/Engine/Plugins/Experimental/DataPrepEditor/Config/BaseDataprepEditor.ini
@@ -1,3 +1,6 @@
 [CoreRedirects]
 
 +ClassRedirects=(OldName="/Script/DataprepEditor.K2Node_DataprepProducer", NewName="/Script/DataprepCore.K2Node_DataprepProducer")
++ClassRedirects=(OldName="DataprepRemoveObjectsOperation", NewName="DataprepDeleteObjectsOperation")
++PropertyRedirects=(OldName="DataprepSetMaterialOperation.MaterialSubstitute",NewName="DataprepSetMaterialOperation.Material")
++PropertyRedirects=(OldName="DataprepSetMeshOperation.MeshSubstitute",NewName="DataprepSetMeshOperation.StaticMesh")
