Fixed issue with material compilation after restoration of a snapshot

diff --git a/Engine/Plugins/Experimental/DataPrepEditor/Source/DataPrepCore/Private/DataprepCoreUtils.cpp b/Engine/Plugins/Experimental/DataPrepEditor/Source/DataPrepCore/Private/DataprepCoreUtils.cpp
--- a/Engine/Plugins/Experimental/DataPrepEditor/Source/DataPrepCore/Private/DataprepCoreUtils.cpp
+++ b/Engine/Plugins/Experimental/DataPrepEditor/Source/DataPrepCore/Private/DataprepCoreUtils.cpp
@@ -323,6 +323,10 @@ void FDataprepCoreUtils::BuildAssets(const TArray<TWeakObjectPtr<UObject>>& Asse
 			if(UMaterialInterface* MaterialInterface = Cast<UMaterialInterface>(AssetObject))
 			{
 				MaterialInterfaces.Add( MaterialInterface );
+				if(MaterialInterface->GetMaterial())
+				{
+					MaterialInterfaces.Add( MaterialInterface->GetMaterial() );
+				}
 			}
 			else if(UStaticMesh* StaticMesh = Cast<UStaticMesh>(AssetObject))
 			{
@@ -351,14 +355,6 @@ void FDataprepCoreUtils::BuildAssets(const TArray<TWeakObjectPtr<UObject>>& Asse
 			{
 				FPropertyChangedEvent EmptyPropertyUpdateStruct( nullptr );
 				MaterialInterface->PostEditChangeProperty( EmptyPropertyUpdateStruct );
-
-			}
-
-			if( MustCompile( MaterialInterface->GetMaterial() ) )
-			{
-				FPropertyChangedEvent EmptyPropertyUpdateStruct( nullptr );
-				MaterialInterface->GetMaterial()->PostEditChangeProperty( EmptyPropertyUpdateStruct );
-
 			}
 		}
 	}
diff --git a/Engine/Plugins/Experimental/DataPrepEditor/Source/DataPrepEditor/Private/DataprepSnapshot.cpp b/Engine/Plugins/Experimental/DataPrepEditor/Source/DataPrepEditor/Private/DataprepSnapshot.cpp
--- a/Engine/Plugins/Experimental/DataPrepEditor/Source/DataPrepEditor/Private/DataprepSnapshot.cpp
+++ b/Engine/Plugins/Experimental/DataPrepEditor/Source/DataPrepEditor/Private/DataprepSnapshot.cpp
@@ -17,6 +17,7 @@
 #include "GenericPlatform/GenericPlatformOutputDevices.h"
 #include "GenericPlatform/GenericPlatformTime.h"
 #include "Materials/MaterialInstance.h"
+#include "MaterialShared.h"
 #include "Misc/Compression.h"
 #include "Misc/FileHelper.h"
 #include "Misc/ScopedSlowTask.h"
@@ -642,6 +643,8 @@ void FDataprepEditor::RestoreFromSnapshot(bool bUpdateViewport)
 
 	SlowTask.EnterProgressFrame( 40.0f, LOCTEXT("RestoreFromSnapshot_Assets", "Restoring assets ...") );
 	{
+		TArray<UMaterialInterface*> MaterialInterfaces;
+
 		FScopedSlowTask SubSlowTask( ContentSnapshot.DataEntries.Num(), LOCTEXT("RestoreFromSnapshot_Assets", "Restoring assets ...") );
 		SubSlowTask.MakeDialog(false);
 
@@ -680,6 +683,18 @@ void FDataprepEditor::RestoreFromSnapshot(bool bUpdateViewport)
 				}
 			}
 
+			if(UMaterialInterface* MaterialInterface = Cast<UMaterialInterface>(Asset))
+			{
+				{
+					FMaterialUpdateContext MaterialUpdateContext;
+
+					MaterialUpdateContext.AddMaterialInterface( MaterialInterface );
+
+					MaterialInterface->PreEditChange(nullptr);
+					MaterialInterface->PostEditChange();
+				}
+			}
+
 			Assets.Add( Asset );
 
 			UE_LOG( LogDataprepEditor, Verbose, TEXT("Asset %s loaded"), *ObjectPath.GetAssetPathString() );
