Fix LevelVariantSets assets not being considered when fixing name collisions in Dataprep #jira UE-81597 #rb JeanMichel.Dignard

diff --git a/Engine/Plugins/Experimental/DataPrepEditor/DataPrepEditor.uplugin b/Engine/Plugins/Experimental/DataPrepEditor/DataPrepEditor.uplugin
--- a/Engine/Plugins/Experimental/DataPrepEditor/DataPrepEditor.uplugin
+++ b/Engine/Plugins/Experimental/DataPrepEditor/DataPrepEditor.uplugin
@@ -35,6 +35,10 @@
 		{
 			"Name": "EditorScriptingUtilities",
 			"Enabled": true
+		},
+		{
+			"Name": "VariantManagerContent",
+			"Enabled": true
 		}
 	]
 }
\ No newline at end of file
diff --git a/Engine/Plugins/Experimental/DataPrepEditor/Source/DataPrepCore/DataPrepCore.Build.cs b/Engine/Plugins/Experimental/DataPrepEditor/Source/DataPrepCore/DataPrepCore.Build.cs
--- a/Engine/Plugins/Experimental/DataPrepEditor/Source/DataPrepCore/DataPrepCore.Build.cs
+++ b/Engine/Plugins/Experimental/DataPrepEditor/Source/DataPrepCore/DataPrepCore.Build.cs
@@ -30,6 +30,7 @@ namespace UnrealBuildTool.Rules
 					"SlateCore",
 					"StaticMeshDescription",
 					"UnrealEd",
+					"VariantManagerContent",
 				}
 			);
 		}
diff --git a/Engine/Plugins/Experimental/DataPrepEditor/Source/DataPrepCore/Private/DataPrepContentProducer.cpp b/Engine/Plugins/Experimental/DataPrepEditor/Source/DataPrepCore/Private/DataPrepContentProducer.cpp
--- a/Engine/Plugins/Experimental/DataPrepEditor/Source/DataPrepCore/Private/DataPrepContentProducer.cpp
+++ b/Engine/Plugins/Experimental/DataPrepEditor/Source/DataPrepCore/Private/DataPrepContentProducer.cpp
@@ -9,6 +9,7 @@
 #include "EngineUtils.h"
 #include "GameFramework/Actor.h"
 #include "LevelSequence.h"
+#include "LevelVariantSets.h"
 
 const FString DefaultNamespace( TEXT("void") );
 
@@ -51,13 +52,20 @@ bool UDataprepContentProducer::Produce(const FDataprepProducerContext& InContext
 		return false;
 	}
 
-	// Collect all packages containing LevelSequence assets to remap their reference to an newly created actor
-	TSet< UPackage* > LevelSequencePackagesToCheck;
+	// Collect all packages containing LevelSequence and LevelVariantSet assets to remap their
+	// reference to an newly created actor
+	TSet< UPackage* > PackagesToCheck;
 	for(int32 Index = LastAssetCount; Index < OutAssets.Num(); ++Index)
 	{
-		if( ULevelSequence* LevelSequence = Cast<ULevelSequence>( OutAssets[Index].Get() ) )
+		UObject* Asset = OutAssets[Index].Get();
+
+		if( ULevelSequence* LevelSequence = Cast<ULevelSequence>(Asset) )
+		{
+			PackagesToCheck.Add( LevelSequence->GetOutermost() );
+		}
+		else if( ULevelVariantSets* ThisLevelVariantSets = Cast<ULevelVariantSets>(Asset) )
 		{
-			LevelSequencePackagesToCheck.Add( LevelSequence->GetOutermost() );
+			PackagesToCheck.Add( ThisLevelVariantSets->GetOutermost() );
 		}
 	}
 
@@ -82,11 +90,11 @@ bool UDataprepContentProducer::Produce(const FDataprepProducerContext& InContext
 		}
 	}
 
-	// Update reference of LevelSequence assets if necessary
-	if(LevelSequencePackagesToCheck.Num() > 0)
+	// Update reference of LevelSequence or LevelVariantSets assets if necessary
+	if(PackagesToCheck.Num() > 0)
 	{
 		IAssetTools& AssetTools = FModuleManager::LoadModuleChecked<FAssetToolsModule>("AssetTools").Get();
-		AssetTools.RenameReferencingSoftObjectPaths( LevelSequencePackagesToCheck.Array(), ActorRedirectorMap );
+		AssetTools.RenameReferencingSoftObjectPaths( PackagesToCheck.Array(), ActorRedirectorMap );
 	}
 
 	Terminate();
