Dataprep parameterization proper ui implemented a proper UI for the creation and removal of bindings.

Added a context menu extender event to the SpinBox and EditableTextBox widget.

diff --git a/Engine/Plugins/Experimental/DataPrepEditor/Source/DataPrepCore/Private/DataPrepAsset.cpp b/Engine/Plugins/Experimental/DataPrepEditor/Source/DataPrepCore/Private/DataPrepAsset.cpp
--- a/Engine/Plugins/Experimental/DataPrepEditor/Source/DataPrepCore/Private/DataPrepAsset.cpp
+++ b/Engine/Plugins/Experimental/DataPrepEditor/Source/DataPrepCore/Private/DataPrepAsset.cpp
@@ -221,11 +221,21 @@ bool UDataprepAsset::IsObjectPropertyBinded(UDataprepParameterizableObject* Obje
 	return Parameterization->IsObjectPropertyBinded( Object, InPropertyChain );
 }
 
+FName UDataprepAsset::GetNameOfParameterForObjectProperty(UDataprepParameterizableObject* Object, const TArray<struct FDataprepPropertyLink>& InPropertyChain) const
+{
+	return Parameterization->GetNameOfParameterForObjectProperty( Object, InPropertyChain );
+}
+
 void UDataprepAsset::RemoveObjectPropertyFromParameterization(UDataprepParameterizableObject* Object, const TArray<struct FDataprepPropertyLink>& InPropertyChain)
 {
 	Parameterization->RemoveBindedObjectProperty( Object, InPropertyChain );
 }
 
+void UDataprepAsset::GetExistingParameterNamesForType(UClass* PropertyClass, TSet<FString>& OutValidExistingNames, TSet<FString>& OutInvalidNames) const
+{
+	Parameterization->GetExistingParameterNamesForType( PropertyClass, OutValidExistingNames, OutInvalidNames );
+}
+
 void UDataprepAsset::OnDataprepBlueprintChanged( UBlueprint* InBlueprint )
 {
 	if(InBlueprint == DataprepRecipeBP)
diff --git a/Engine/Plugins/Experimental/DataPrepEditor/Source/DataPrepCore/Private/Parameterization/DataprepParameterization.cpp b/Engine/Plugins/Experimental/DataPrepEditor/Source/DataPrepCore/Private/Parameterization/DataprepParameterization.cpp
--- a/Engine/Plugins/Experimental/DataPrepEditor/Source/DataPrepCore/Private/Parameterization/DataprepParameterization.cpp
+++ b/Engine/Plugins/Experimental/DataPrepEditor/Source/DataPrepCore/Private/Parameterization/DataprepParameterization.cpp
@@ -679,10 +679,7 @@ void UDataprepParameterization::PostEditUndo()
 
 	OnTellInstancesToReloadTheirSerializedData.Broadcast();
 
-	if ( UDataprepAsset* DataprepAsset = Cast<UDataprepAsset>( GetOuter() ) )
-	{
-		DataprepAsset->OnParameterizationWasModified.Broadcast();
-	}
+	CastChecked<UDataprepAsset>( GetOuter() )->OnParameterizationStatusForObjectsChanged.Broadcast( nullptr );
 }
 
 void UDataprepParameterization::OnObjectModified(UObject* Object)
@@ -705,6 +702,20 @@ bool UDataprepParameterization::BindObjectProperty(UDataprepParameterizableObjec
 		TSharedRef<FDataprepParameterizationBinding> Binding = MakeShared<FDataprepParameterizationBinding>( Object, PropertyChain );
 		void* AddressOfTheValueFromBinding;
 
+		FName OldParameterName = BindingsContainer->GetParameterNameForBinding( Binding );
+		if ( OldParameterName == Name )
+		{
+			return !Name.IsNone();
+		}
+
+		bool bClassNeedUpdate = false;
+		if ( OldParameterName != Name )
+		{
+			RemoveBinding( Binding, bClassNeedUpdate );
+		}
+
+		bool bBindingWasAdded = false;
+
 		// We expect the chain to have a valid chain of cached property before inserting the binding
 		if ( DataprepParameterization::GetDeepestLevelOfValidCache( Binding.Get(), AddressOfTheValueFromBinding) == PropertyChain.Num() - 1 )
 		{
@@ -717,7 +728,7 @@ bool UDataprepParameterization::BindObjectProperty(UDataprepParameterizableObjec
 				{
 					Modify();
 					BindingsContainer->Add( Binding, Name );
-					UpdateParameterizationFromBinding( Binding );
+					bBindingWasAdded = true;
 				}
 			}
 			else
@@ -730,11 +741,25 @@ bool UDataprepParameterization::BindObjectProperty(UDataprepParameterizableObjec
 
 				// The validation we did with GetDeepestLevelOfValidCache ensure us that the property ptr is valid
 				UProperty* NewProperty = AddPropertyToClass( Name, *SourceProperty );
-				UpdateClass();
 
-				UpdateParameterizationFromBinding( Binding );
+				bClassNeedUpdate = true;
+				bBindingWasAdded = true;
 			}
 		}
+
+		if ( bClassNeedUpdate )
+		{
+			UpdateClass();
+		}
+
+		if ( bBindingWasAdded )
+		{
+			UpdateParameterizationFromBinding( Binding );
+
+			TSet<UObject*> Objects;
+			Objects.Add( Object );
+			CastChecked<UDataprepAsset>( GetOuter() )->OnParameterizationStatusForObjectsChanged.Broadcast( &Objects );
+		}
 	}
 
 	return false;
@@ -746,17 +771,29 @@ bool UDataprepParameterization::IsObjectPropertyBinded(UDataprepParameterizableO
 	return BindingsContainer->ContainsBinding( Binding );
 }
 
+FName UDataprepParameterization::GetNameOfParameterForObjectProperty(UDataprepParameterizableObject* Object, const TArray<struct FDataprepPropertyLink>& PropertyChain) const
+{
+	TSharedRef<FDataprepParameterizationBinding> Binding = MakeShared<FDataprepParameterizationBinding>(Object, PropertyChain);
+	return BindingsContainer->GetParameterNameForBinding(Binding);
+}
+
 void UDataprepParameterization::RemoveBindedObjectProperty(UDataprepParameterizableObject* Object, const TArray<FDataprepPropertyLink>& PropertyChain)
 {
 	TSharedRef<FDataprepParameterizationBinding> Binding = MakeShared<FDataprepParameterizationBinding>( Object, PropertyChain );
 
 	Modify();
 
-	FName ParameterOfRemovedBinding = BindingsContainer->RemoveBinding( Binding );
-	if ( !BindingsContainer->HasBindingsForParameter( ParameterOfRemovedBinding ) )
+	bool bClassWasModified;
+	if ( RemoveBinding( Binding, bClassWasModified ) )
 	{
-		NameToParameterizationProperty.Remove( ParameterOfRemovedBinding );
-		UpdateClass();
+		if ( bClassWasModified )
+		{
+			UpdateClass();
+		}
+
+		TSet<UObject*> Objects;
+		Objects.Add( Object );
+		CastChecked<UDataprepAsset>( GetOuter() )->OnParameterizationStatusForObjectsChanged.Broadcast( &Objects );
 	}
 }
 
@@ -764,40 +801,51 @@ void UDataprepParameterization::RemoveBindingFromObjects(TArray<UDataprepParamet
 {
 	Modify();
 	TSet<FName> ParameterPotentiallyRemoved;
+	TSet<UObject*> UniqueObjects;
+	UniqueObjects.Reserve( Objects.Num() );
+
 	for ( UDataprepParameterizableObject* Object : Objects )
 	{
 		ParameterPotentiallyRemoved.Append( BindingsContainer->RemoveAllBindingsFromObject( Object ) );
+		UniqueObjects.Add( Object );
 	}
 
+	bool bClassWasChanged = false;
 	for ( const FName& Name : ParameterPotentiallyRemoved )
 	{
 		if ( !BindingsContainer->HasBindingsForParameter( Name ) )
 		{
 			NameToParameterizationProperty.Remove( Name );
+			bClassWasChanged = true;
 		}
 	}
 
-	UpdateClass();
+	if ( bClassWasChanged )
+	{
+		UpdateClass();
+	}
+
+	CastChecked<UDataprepAsset>( GetOuter() )->OnParameterizationStatusForObjectsChanged.Broadcast( &UniqueObjects );
 }
 
 void UDataprepParameterization::UpdateParameterizationFromBinding(const TSharedRef<FDataprepParameterizationBinding>& Binding)
 {
 	FName ParameterModified = BindingsContainer->GetParameterNameForBinding(Binding);
-	UProperty* ParametrizationProperty = NameToParameterizationProperty.FindRef(ParameterModified);
-	if (ParametrizationProperty)
+	UProperty* ParameterizationProperty = NameToParameterizationProperty.FindRef(ParameterModified);
+	if (ParameterizationProperty)
 	{
 		void* AddressOfObjectValue = nullptr;
 		if (UProperty * ObjectProperty = DataprepParameterization::GetPropertyFromBinding(Binding.Get(), AddressOfObjectValue))
 		{
 			Modify();
-			void* AddressOfParameterizationValue = DataprepParameterization::GetAddressOf(*ParametrizationProperty, DefaultParameterisation, INDEX_NONE);
-			DataprepParameterization::CopyCompleteValue(*ParametrizationProperty, AddressOfParameterizationValue, *ObjectProperty, AddressOfObjectValue);
+			void* AddressOfParameterizationValue = DataprepParameterization::GetAddressOf(*ParameterizationProperty, DefaultParameterisation, INDEX_NONE);
+			DataprepParameterization::CopyCompleteValue(*ParameterizationProperty, AddressOfParameterizationValue, *ObjectProperty, AddressOfObjectValue);
 
 			// Post edit the default parameterization
 			FEditPropertyChain EditChain;
-			EditChain.AddHead(ParametrizationProperty);
-			EditChain.SetActivePropertyNode(ParametrizationProperty);
-			FPropertyChangedEvent EditPropertyChangeEvent(ParametrizationProperty, EPropertyChangeType::ValueSet);
+			EditChain.AddHead(ParameterizationProperty);
+			EditChain.SetActivePropertyNode(ParameterizationProperty);
+			FPropertyChangedEvent EditPropertyChangeEvent(ParameterizationProperty, EPropertyChangeType::ValueSet);
 			FPropertyChangedChainEvent EditChangeChainEvent(EditChain, EditPropertyChangeEvent);
 			DefaultParameterisation->PostEditChangeChainProperty(EditChangeChainEvent);
 		}
@@ -824,6 +872,27 @@ void UDataprepParameterization::OnObjectPostEdit(UDataprepParameterizableObject*
 	}
 }
 
+void UDataprepParameterization::GetExistingParameterNamesForType(UClass* PropertyClass, TSet<FString>& OutValidExistingNames, TSet<FString>& OutInvalidNames) const
+{
+	OutValidExistingNames.Empty( NameToParameterizationProperty.Num() );
+	OutInvalidNames.Empty( NameToParameterizationProperty.Num() );
+
+	for ( const TPair<FName, UProperty*>& Pair : NameToParameterizationProperty )
+	{
+		if ( const UProperty* Property = Pair.Value )
+		{
+			if ( Property->GetClass() == PropertyClass )
+			{
+				OutValidExistingNames.Add( Pair.Key.ToString() );
+			}
+			else
+			{
+				OutInvalidNames.Add( Pair.Key.ToString() );
+			}
+		}
+	}
+}
+
 void UDataprepParameterization::GenerateClass()
 {
 	if ( !CustomContainerClass )
@@ -905,9 +974,11 @@ void UDataprepParameterization::LoadParameterization()
 		}
 
 		// Remove the invalid bindings
+		TSet<UObject*> ObjectsToNotify;
 		for ( const TSharedRef<FDataprepParameterizationBinding>& InvalidBinding : BindingToRemove )
 		{
 			BindingsContainer->RemoveBinding( InvalidBinding );
+			ObjectsToNotify.Add( InvalidBinding->ObjectBinded );
 		}
 
 		// Make the properties appear in a alphabetically order (for that we must add the properties to the class in the reverse order)
@@ -927,6 +998,11 @@ void UDataprepParameterization::LoadParameterization()
 
 		DefaultParameterisation = CustomContainerClass->GetDefaultObject(true);
 		FDataprepParameterizationReader Reader( DefaultParameterisation, ParameterizationStorage );
+
+		if ( ObjectsToNotify.Num() > 0 )
+		{
+			CastChecked<UDataprepAsset>( GetOuter() )->OnParameterizationStatusForObjectsChanged.Broadcast( &ObjectsToNotify );
+		}
 	}
 }
 
@@ -1055,17 +1131,45 @@ void UDataprepParameterization::PushParametrizationValueToBindings(FName Paramet
 					}
 				}
 
+				TSet<UObject*> ObjectsToNotify;
+				bool bClassNeedUpdate = false;
+
 				// Remove the invalid bindings
 				for ( const TSharedRef<FDataprepParameterizationBinding>& Binding : BindingToRemove )
 				{
-					// Naive slow implementation for now
-					RemoveBindedObjectProperty( Binding->ObjectBinded, Binding->PropertyChain );
+					bool bModifiedClass;
+					if ( RemoveBinding( Binding, bModifiedClass ) )
+					{
+						bClassNeedUpdate |= bModifiedClass;
+						ObjectsToNotify.Add( Binding->ObjectBinded );
+					}
+				}
+
+				if ( bClassNeedUpdate )
+				{
+					UpdateClass();
+				}
+
+				if ( ObjectsToNotify.Num() > 0 )
+				{
+					CastChecked<UDataprepAsset>( GetOuter() )->OnParameterizationStatusForObjectsChanged.Broadcast( &ObjectsToNotify );
 				}
 			}
 		}
 	}
 }
 
+bool UDataprepParameterization::RemoveBinding(const TSharedRef<FDataprepParameterizationBinding>& Binding, bool& bOutClassNeedUpdate)
+{
+	FName ParameterOfRemovedBinding = BindingsContainer->RemoveBinding( Binding );
+	if ( !BindingsContainer->HasBindingsForParameter( ParameterOfRemovedBinding ) )
+	{
+		bOutClassNeedUpdate = true;
+		NameToParameterizationProperty.Remove(ParameterOfRemovedBinding);
+	}
+	return !ParameterOfRemovedBinding.IsNone();
+}
+
 const FName UDataprepParameterization::MetadataClassGeneratorName( TEXT("DataprepCustomParameterizationGenerator") );
 
 
diff --git a/Engine/Plugins/Experimental/DataPrepEditor/Source/DataPrepCore/Private/Parameterization/DataprepParameterization.h b/Engine/Plugins/Experimental/DataPrepEditor/Source/DataPrepCore/Private/Parameterization/DataprepParameterization.h
--- a/Engine/Plugins/Experimental/DataPrepEditor/Source/DataPrepCore/Private/Parameterization/DataprepParameterization.h
+++ b/Engine/Plugins/Experimental/DataPrepEditor/Source/DataPrepCore/Private/Parameterization/DataprepParameterization.h
@@ -32,8 +32,16 @@ struct FDataprepParameterizationBinding
 	FDataprepParameterizationBinding(UDataprepParameterizableObject* InObjectBinded, TArray<FDataprepPropertyLink> InPropertyChain)
 		: ObjectBinded( InObjectBinded )
 		, PropertyChain( MoveTemp( InPropertyChain ) )
-		, ValueType( PropertyChain.Last().CachedProperty.Get() ? PropertyChain.Last().CachedProperty.Get()->GetClass() : nullptr )
-	{}
+		, ValueType( nullptr )
+	{
+		if ( PropertyChain.Num() > 0 )
+		{
+			if ( UProperty* Property = PropertyChain.Last().CachedProperty.Get() )
+			{
+				ValueType = Property->GetClass();
+			}
+		}
+	}
 
 	FDataprepParameterizationBinding(FDataprepParameterizationBinding&&) = default;
 	FDataprepParameterizationBinding(const FDataprepParameterizationBinding&) = default;
@@ -204,12 +212,16 @@ public:
 
 	bool IsObjectPropertyBinded(UDataprepParameterizableObject* Object, const TArray<FDataprepPropertyLink>& PropertyChain) const;
 
+	FName GetNameOfParameterForObjectProperty(UDataprepParameterizableObject* Object, const TArray<struct FDataprepPropertyLink>& PropertyChain) const;
+
 	void RemoveBindedObjectProperty(UDataprepParameterizableObject* Object, const TArray<FDataprepPropertyLink>& PropertyChain);
 
 	void RemoveBindingFromObjects(TArray<UDataprepParameterizableObject*> Objects);
 
 	void OnObjectPostEdit(UDataprepParameterizableObject* Object, const TArray<FDataprepPropertyLink>& PropertyChain, EPropertyChangeType::Type ChangeType);
 
+	void GetExistingParameterNamesForType(UClass* PropertyClass,TSet<FString>& OutValidExistingNames, TSet<FString>& OutInvalidNames) const;
+
 private:
 
 	/**
@@ -260,6 +272,12 @@ private:
 	 */
 	void PushParametrizationValueToBindings(FName ParameterName);
 
+	/**
+	 * Do the actual removing of a binding
+	 * @return True if binding was remove
+	 */
+	bool RemoveBinding(const TSharedRef<FDataprepParameterizationBinding>& Binding, bool& bOutClassNeedUpdate);
+
 public:
 
 	static const FName MetadataClassGeneratorName;
diff --git a/Engine/Plugins/Experimental/DataPrepEditor/Source/DataPrepCore/Public/DataPrepAsset.h b/Engine/Plugins/Experimental/DataPrepEditor/Source/DataPrepCore/Public/DataPrepAsset.h
--- a/Engine/Plugins/Experimental/DataPrepEditor/Source/DataPrepCore/Public/DataPrepAsset.h
+++ b/Engine/Plugins/Experimental/DataPrepEditor/Source/DataPrepCore/Public/DataPrepAsset.h
@@ -96,11 +96,12 @@ public:
 	// Functions specific to the parametrization of the dataprep asset
 
 	/**
-	 * Delegate to notify the ui that a dataprep parametrization was modified
-	 * This necessary as the ui for the parameterization is only updated by manual event
+	 * Event to notify the ui that a dataprep parametrization was modified
+	 * This necessary as the ui for the parameterization is only updated by manual event (the ui don't pull new values each frame)
+	 * Note on the objects param: The parameterized objects that should refresh their ui. If nullptr all widgets that can display some info on the parameterization should be refreshed
 	 */
-	DECLARE_MULTICAST_DELEGATE(FDataprepParameterizationWasModified)
-	FDataprepParameterizationWasModified OnParameterizationWasModified;
+	DECLARE_EVENT_OneParam(UDataprepParameterization, FDataprepParameterizationStatusForObjectsChanged, const TSet<UObject*>* /** Objects */)
+	FDataprepParameterizationStatusForObjectsChanged OnParameterizationStatusForObjectsChanged;
 
 	virtual UObject* GetParameterizationObject() override;
 
@@ -108,8 +109,12 @@ public:
 
 	bool IsObjectPropertyBinded(UDataprepParameterizableObject* Object, const TArray<struct FDataprepPropertyLink>& InPropertyChain) const;
 
+	FName GetNameOfParameterForObjectProperty(UDataprepParameterizableObject* Object, const TArray<struct FDataprepPropertyLink>& InPropertyChain) const;
+
 	void RemoveObjectPropertyFromParameterization(UDataprepParameterizableObject* Object, const TArray<struct FDataprepPropertyLink>& InPropertyChain);
 
+	void GetExistingParameterNamesForType(UClass* PropertyClass, TSet<FString>& OutValidExistingNames, TSet<FString>& OutInvalidNames) const;
+
 	UDataprepParameterization* GetDataprepParameterization() { return Parameterization; }
 
 protected:
diff --git a/Engine/Plugins/Experimental/DataPrepEditor/Source/DataPrepEditor/Private/DataprepEditorUtils.cpp b/Engine/Plugins/Experimental/DataPrepEditor/Source/DataPrepEditor/Private/DataprepEditorUtils.cpp
--- a/Engine/Plugins/Experimental/DataPrepEditor/Source/DataPrepEditor/Private/DataprepEditorUtils.cpp
+++ b/Engine/Plugins/Experimental/DataPrepEditor/Source/DataPrepEditor/Private/DataprepEditorUtils.cpp
@@ -3,10 +3,29 @@
 #include "DataprepEditorUtils.h"
 
 #include "DataPrepAsset.h"
+#include "DataprepParameterizableObject.h"
+#include "Parameterization/DataprepParameterizationUtils.h"
+#include "Widgets/Parameterization/SDataprepLinkToParameter.h"
 
 #include "EditorStyleSet.h"
 #include "Engine/Blueprint.h"
+#include "Framework/MultiBox/MultiBoxBuilder.h"
 #include "Kismet2/BlueprintEditorUtils.h"
+#include "ScopedTransaction.h"
+#include "Widgets/Text/STextBlock.h"
+
+#define LOCTEXT_NAMESPACE "DataprepEditorUtils"
+
+void FDataprepParametrizationActionData::AddReferencedObjects(FReferenceCollector& Collector)
+{
+	Collector.AddReferencedObject(DataprepAsset);
+	Collector.AddReferencedObject(Object);
+}
+
+bool FDataprepParametrizationActionData::IsValid() const
+{
+	return DataprepAsset && Object && PropertyChain.Num() > 0;
+}
 
 void FDataprepEditorUtils::NotifySystemOfChangeInPipeline(UObject* SourceObject)
 {
@@ -38,8 +57,63 @@ void FDataprepEditorUtils::NotifySystemOfChangeInPipeline(UObject* SourceObject)
 	}
 }
 
+void FDataprepEditorUtils::PopulateMenuForParameterization(FMenuBuilder& MenuBuilder, UDataprepAsset& DataprepAsset, UDataprepParameterizableObject& Object, const TArray<FDataprepPropertyLink>& PropertyChain)
+{
+	TSharedRef<FDataprepParametrizationActionData> ActionData = MakeShared<FDataprepParametrizationActionData>( DataprepAsset, Object, PropertyChain );
+
+	MenuBuilder.BeginSection( NAME_None, LOCTEXT("ParametrizationMenuSection", "Parameterization") );
+	{
+		FName ParameterName = DataprepAsset.GetNameOfParameterForObjectProperty( &Object, PropertyChain );
+
+		FNewMenuDelegate BindToParamerizationDelegate;
+		
+
+		BindToParamerizationDelegate.BindLambda( [ActionData, ParameterName](FMenuBuilder& BindToParamerizationDelegateMenu)
+			{
+				FUIAction DoNothing;
+				BindToParamerizationDelegateMenu.AddWidget( SNew(SDataprepLinkToParameter, ActionData), FText(), true, false );
+			});
+
+
+		MenuBuilder.AddSubMenu( LOCTEXT("LinkToParameterLabel", "Link To Parameter"), LOCTEXT("LinkToParameterTooltip", "Link this property to a existing parameter or a new one")
+			, BindToParamerizationDelegate, true, FSlateIcon(), false );
+
+		if ( !ParameterName.IsNone() )
+		{ 
+			
+			FUIAction RemoveBinding;
+			FText UnlinkFromParameterLabel = LOCTEXT("UnlinkFromParameterLabel", "Remove Link To Parameter");
+			RemoveBinding.ExecuteAction.BindLambda( [ActionData, UnlinkFromParameterLabel]()
+				{
+					if ( ActionData->IsValid() )
+					{
+						FScopedTransaction ScopedTransaction( UnlinkFromParameterLabel );
+						ActionData->DataprepAsset->RemoveObjectPropertyFromParameterization( ActionData->Object, ActionData->PropertyChain );
+					}
+				});
+
+			MenuBuilder.AddMenuEntry( UnlinkFromParameterLabel, TAttribute<FText>(), FSlateIcon(), RemoveBinding );
+		}
+	}
+	MenuBuilder.EndSection();
+}
+
 FSlateFontInfo FDataprepEditorUtils::GetGlyphFont()
 {
-	return FEditorStyle::Get().GetFontStyle( "FontAwesome.11" );
+	return FEditorStyle::Get().GetFontStyle("FontAwesome.11");
+}
+
+TSharedPtr<SWidget> FDataprepEditorUtils::MakeContextMenu(const TSharedPtr<FDataprepParametrizationActionData>& ParameterizationActionData)
+{
+	if ( ParameterizationActionData && ParameterizationActionData->IsValid() )
+	{
+		FMenuBuilder MenuBuilder( true, nullptr );
+		PopulateMenuForParameterization( MenuBuilder, *ParameterizationActionData->DataprepAsset,
+			*ParameterizationActionData->Object, ParameterizationActionData->PropertyChain );
+		return MenuBuilder.MakeWidget();
+	}
+
+	return {};
 }
 
+#undef LOCTEXT_NAMESPACE
diff --git a/Engine/Plugins/Experimental/DataPrepEditor/Source/DataPrepEditor/Private/DataprepEditorUtils.h b/Engine/Plugins/Experimental/DataPrepEditor/Source/DataPrepEditor/Private/DataprepEditorUtils.h
--- a/Engine/Plugins/Experimental/DataPrepEditor/Source/DataPrepEditor/Private/DataprepEditorUtils.h
+++ b/Engine/Plugins/Experimental/DataPrepEditor/Source/DataPrepEditor/Private/DataprepEditorUtils.h
@@ -3,12 +3,53 @@
 #pragma once
 
 #include "CoreMinimal.h"
-
 #include "Fonts/SlateFontInfo.h"
 
+class FMenuBuilder;
+class SWidget;
+class UDataprepAsset;
+class UDataprepParameterizableObject;
+
+struct FDataprepPropertyLink;
+
+struct FDataprepParametrizationActionData : public FGCObject
+{
+	FDataprepParametrizationActionData(UDataprepAsset& InDataprepAsset, UDataprepParameterizableObject& InObject, const TArray<FDataprepPropertyLink>& InPropertyChain)
+		: FGCObject()
+		, DataprepAsset(&InDataprepAsset)
+		, Object(&InObject)
+		, PropertyChain(InPropertyChain)
+	{}
+
+	UDataprepAsset* DataprepAsset;
+	UDataprepParameterizableObject* Object;
+	TArray<FDataprepPropertyLink> PropertyChain;
+
+	virtual void AddReferencedObjects(FReferenceCollector& Collector) override;
+
+	bool IsValid() const;
+};
+
 class FDataprepEditorUtils
 {
 public:
+	/**
+	 * Find the the owning dataprep asset of the source object and send a notification to the dataprep editor so it can react when its pipeline is modified
+	 */
 	static void NotifySystemOfChangeInPipeline(UObject* SourceObject);
+
+	/**
+	 * Populate a menu builder with the section made for the parameterization
+	 * @param DataprepAsset the asset that own the object
+	 * @param Object The Object on which we want to modify the parametrization binding
+	 * @param PropertyChain The property chain is the property path from the class of the object to the property that we want to edit
+	 */
+	static void PopulateMenuForParameterization(FMenuBuilder& MenuBuilder, UDataprepAsset& DataprepAsset, UDataprepParameterizableObject& Object, const TArray<FDataprepPropertyLink>& PropertyChain);
+
 	static FSlateFontInfo GetGlyphFont();
+
+	/**
+	 * Make a context menu widget to manage parameterization link of a property
+	 */
+	static TSharedPtr<SWidget> MakeContextMenu(const TSharedPtr<FDataprepParametrizationActionData>& ParameterizationActionData);
 };
diff --git a/Engine/Plugins/Experimental/DataPrepEditor/Source/DataPrepEditor/Private/Widgets/Action/SDataprepActionBlock.cpp b/Engine/Plugins/Experimental/DataPrepEditor/Source/DataPrepEditor/Private/Widgets/Action/SDataprepActionBlock.cpp
--- a/Engine/Plugins/Experimental/DataPrepEditor/Source/DataPrepEditor/Private/Widgets/Action/SDataprepActionBlock.cpp
+++ b/Engine/Plugins/Experimental/DataPrepEditor/Source/DataPrepEditor/Private/Widgets/Action/SDataprepActionBlock.cpp
@@ -15,6 +15,7 @@
 #include "Framework/Commands/UICommandList.h"
 #include "Framework/MultiBox/MultiBoxBuilder.h"
 #include "Framework/Text/TextLayout.h"
+#include "Layout/WidgetPath.h"
 #include "ScopedTransaction.h"
 #include "Styling/ISlateStyle.h"
 #include "Styling/SlateStyleRegistry.h"
@@ -124,9 +125,10 @@ FReply SDataprepActionBlock::OnMouseButtonUp(const FGeometry& MyGeometry, const
 		FMenuBuilder MenuBuilder( true, nullptr );
 		PopulateMenuBuilder( MenuBuilder );
 
-		TSharedPtr<IMenu> Menu = FSlateApplication::Get().PushMenu(
+		FWidgetPath WidgetPath = MouseEvent.GetEventPath() ? *MouseEvent.GetEventPath() : FWidgetPath();
+		FSlateApplication::Get().PushMenu(
 			AsShared(),
-			FWidgetPath(),
+			WidgetPath,
 			MenuBuilder.MakeWidget(),
 			MouseEvent.GetScreenSpacePosition(),
 			FPopupTransitionEffect(FPopupTransitionEffect::ContextMenu)
diff --git a/Engine/Plugins/Experimental/DataPrepEditor/Source/DataPrepEditor/Private/Widgets/Action/SDataprepFloatFilter.cpp b/Engine/Plugins/Experimental/DataPrepEditor/Source/DataPrepEditor/Private/Widgets/Action/SDataprepFloatFilter.cpp
--- a/Engine/Plugins/Experimental/DataPrepEditor/Source/DataPrepEditor/Private/Widgets/Action/SDataprepFloatFilter.cpp
+++ b/Engine/Plugins/Experimental/DataPrepEditor/Source/DataPrepEditor/Private/Widgets/Action/SDataprepFloatFilter.cpp
@@ -2,12 +2,16 @@
 
 #include "Widgets/Action/SDataprepFloatFilter.h"
 
+#include "DataPrepAsset.h"
+#include "DataprepCoreUtils.h"
 #include "DataprepEditorUtils.h"
 #include "SchemaActions/DataprepFetcherMenuActionCollector.h"
 #include "SchemaActions/DataprepSchemaAction.h"
 #include "SelectionSystem/DataprepFloatFilter.h"
 #include "Widgets/Action/DataprepActionWidgetsUtils.h"
 #include "Widgets/Action/SDataprepFetcherSelector.h"
+#include "Widgets/DataprepWidgets.h"
+#include "Widgets/Parameterization/SDataprepParameterizationLinkIcon.h"
 
 #include "Internationalization/Text.h"
 #include "ScopedTransaction.h"
@@ -19,6 +23,25 @@
 
 #define LOCTEXT_NAMESPACE "DataprepFloatFilter"
 
+namespace SDataprepFloatFilterUtils
+{
+	void PostEditChainProperty(UDataprepFloatFilter* Filter, FName PropertyName)
+	{
+		if ( Filter )
+		{
+			UProperty* Property = Filter->GetClass()->FindPropertyByName( PropertyName );
+			check( Property );
+
+			FEditPropertyChain EditChain;
+			EditChain.AddHead( Property );
+			EditChain.SetActivePropertyNode( Property );
+			FPropertyChangedEvent EditPropertyChangeEvent( Property, EPropertyChangeType::ValueSet );
+			FPropertyChangedChainEvent EditChangeChainEvent( EditChain, EditPropertyChangeEvent );
+			Filter->PostEditChangeChainProperty( EditChangeChainEvent );
+		}
+	}
+};
+
 void SDataprepFloatFilter::Construct(const FArguments& InArgs, UDataprepFloatFilter& InFilter)
 {
 	Filter = &InFilter;
@@ -27,6 +50,60 @@ void SDataprepFloatFilter::Construct(const FArguments& InArgs, UDataprepFloatFil
 
 	DataprepActionWidgetsUtils::GenerateListEntriesFromEnum< EDataprepFloatMatchType >( FloatMatchingOptions );
 
+	if ( UDataprepAsset* DataprepAsset = FDataprepCoreUtils::GetDataprepAssetOfObject( Filter ) )
+	{
+		UClass* FilterClass = InFilter.GetClass();
+
+		{
+			FName PropertyName = FName( TEXT("FloatMatchingCriteria") );
+			UProperty* Property = FilterClass->FindPropertyByName( PropertyName );
+			check( Property );
+			TArray<FDataprepPropertyLink> PropertyChain;
+			PropertyChain.Emplace( Property, PropertyName, INDEX_NONE );
+
+			MatchingCriteriaParameterizationActionData = MakeShared<FDataprepParametrizationActionData>( *DataprepAsset, InFilter, PropertyChain );
+		}
+
+		{
+			FName PropertyName = FName( TEXT("EqualValue") );
+			UProperty* Property = FilterClass->FindPropertyByName( PropertyName );
+			check( Property );
+			TArray<FDataprepPropertyLink> PropertyChain;
+			PropertyChain.Emplace( Property, PropertyName, INDEX_NONE );
+
+			EqualValueParameterizationActionData = MakeShared<FDataprepParametrizationActionData>( *DataprepAsset, InFilter, PropertyChain );
+		}
+
+		{
+			FName PropertyName = FName( TEXT("Tolerance") );
+			UProperty* Property = FilterClass->FindPropertyByName( PropertyName );
+			check( Property );
+			TArray<FDataprepPropertyLink> PropertyChain;
+			PropertyChain.Emplace( Property, PropertyName, INDEX_NONE );
+
+			ToleranceParameterizationActionData = MakeShared<FDataprepParametrizationActionData>( *DataprepAsset, InFilter, PropertyChain );
+		}
+
+		OnParameterizationStatusForObjectsChangedHandle = DataprepAsset->OnParameterizationStatusForObjectsChanged.AddSP( this, &SDataprepFloatFilter::OnParameterizationStatusForObjectsChanged );
+	}
+
+	UpdateVisualDisplay();
+}
+
+SDataprepFloatFilter::~SDataprepFloatFilter()
+{
+	if ( UDataprepAsset* DataprepAsset = FDataprepCoreUtils::GetDataprepAssetOfObject( Filter ) )
+	{
+		DataprepAsset->OnParameterizationStatusForObjectsChanged.Remove( OnParameterizationStatusForObjectsChangedHandle );
+	}
+}
+
+void SDataprepFloatFilter::UpdateVisualDisplay()
+{
+	TSharedPtr<SHorizontalBox> MatchingCriteriaHorizontalBox;
+	TSharedPtr<SHorizontalBox> EqualValueHorizontalBox;
+	TSharedPtr<SHorizontalBox> ToleranceHorizontalBox;
+
 	ChildSlot
 	[
 		SNew( SBox )
@@ -40,60 +117,130 @@ void SDataprepFloatFilter::Construct(const FArguments& InArgs, UDataprepFloatFil
 				+ SHorizontalBox::Slot()
 				.Padding( 5.f )
 				[
-					SAssignNew( FloatMatchingCriteriaWidget, SComboBox< TSharedPtr< FListEntry > > )
-					.OptionsSource( &FloatMatchingOptions )
-					.OnGenerateWidget( this, &SDataprepFloatFilter::OnGenerateWidgetForMatchingCriteria )
-					.OnSelectionChanged( this, &SDataprepFloatFilter::OnSelectedCriteriaChanged )
-					.OnComboBoxOpening( this, &SDataprepFloatFilter::OnCriteriaComboBoxOpenning )
+					SNew( SDataprepContextMenuOverride )
+					.OnContextMenuOpening( this, &SDataprepFloatFilter::OnGetContextMenuForMatchingCriteria )
 					[
-						SNew( STextBlock )
-						.Text( this, &SDataprepFloatFilter::GetSelectedCriteriaText )
-						.ToolTipText( this, &SDataprepFloatFilter::GetSelectedCriteriaTooltipText )
-						.Justification( ETextJustify::Center )
+						SAssignNew( MatchingCriteriaHorizontalBox, SHorizontalBox )
+						+ SHorizontalBox::Slot()
+						[
+							SAssignNew( FloatMatchingCriteriaWidget, SComboBox< TSharedPtr< FListEntry > > )
+							.OptionsSource( &FloatMatchingOptions )
+							.OnGenerateWidget( this, &SDataprepFloatFilter::OnGenerateWidgetForMatchingCriteria )
+							.OnSelectionChanged( this, &SDataprepFloatFilter::OnSelectedCriteriaChanged )
+							.OnComboBoxOpening( this, &SDataprepFloatFilter::OnCriteriaComboBoxOpenning )
+							[
+								SNew( STextBlock )
+								.Text( this, &SDataprepFloatFilter::GetSelectedCriteriaText )
+								.ToolTipText( this, &SDataprepFloatFilter::GetSelectedCriteriaTooltipText )
+								.Justification( ETextJustify::Center )
+							]
+						]
 					]
-			
 				]
 				+ SHorizontalBox::Slot()
 				.Padding( 5.f )
 				[
-					SNew( SSpinBox<float> )
-					.Value( this, &SDataprepFloatFilter::GetEqualValue )
-					.OnValueChanged( this, &SDataprepFloatFilter::OnEqualValueChanged )
-					.OnValueCommitted( this, &SDataprepFloatFilter::OnEqualValueComitted )
-					.Justification( ETextJustify::Center )
-					.MinValue( TOptional< float >() )
-					.MaxValue( TOptional< float >() )
+					SNew( SDataprepContextMenuOverride )
+					.OnContextMenuOpening( this, &SDataprepFloatFilter::OnGetContextMenuForEqualValue )
+					[
+						SAssignNew( EqualValueHorizontalBox, SHorizontalBox )
+						+ SHorizontalBox::Slot()
+						[
+							SNew( SSpinBox<float> )
+							.Value( this, &SDataprepFloatFilter::GetEqualValue )
+							.OnValueChanged( this, &SDataprepFloatFilter::OnEqualValueChanged )
+							.OnValueCommitted( this, &SDataprepFloatFilter::OnEqualValueComitted )
+							.Justification( ETextJustify::Center )
+							.MinValue( TOptional< float >() )
+							.MaxValue( TOptional< float >() )
+							.ContextMenuExtender( this, &SDataprepFloatFilter::ExtendContextMenuForEqualValueBox )
+						]
+					]
 				]
 			]
 			+ SVerticalBox::Slot()
 			.AutoHeight()
 			[
-				SNew(SHorizontalBox)
-				.Visibility( this, &SDataprepFloatFilter::GetToleranceRowVisibility )
-				+ SHorizontalBox::Slot()
-				.AutoWidth()
-				.VAlign( VAlign_Center )
-				.Padding( 5.f )
-				[
-					SNew( STextBlock )
-					.Text( LOCTEXT("ToleranceText", "Tolerance") )
-					.ToolTipText(this, &SDataprepFloatFilter::GetSelectedCriteriaTooltipText)
-					.Justification( ETextJustify::Center )
-				]
-				+ SHorizontalBox::Slot()
-				.Padding( 5.f )
+				SNew( SDataprepContextMenuOverride )
+				.OnContextMenuOpening( this, &SDataprepFloatFilter::OnGetContextMenuForTolerance )
 				[
-					SNew( SSpinBox<float> )
-					.Value( this, &SDataprepFloatFilter::GetTolerance )
-					.OnValueChanged( this, &SDataprepFloatFilter::OnToleranceChanged )
-					.OnValueCommitted( this, &SDataprepFloatFilter::OnToleranceComitted )
-					.Justification( ETextJustify::Center )
-					.MinValue( TOptional< float >() )
-					.MaxValue( TOptional< float >() )
+					SNew(SHorizontalBox)
+					.Visibility( this, &SDataprepFloatFilter::GetToleranceRowVisibility )
+					+ SHorizontalBox::Slot()
+					.AutoWidth()
+					.VAlign( VAlign_Center )
+					.Padding( 5.f )
+					[
+						SAssignNew( ToleranceHorizontalBox, SHorizontalBox )
+						+ SHorizontalBox::Slot()
+						[
+							SNew( STextBlock )
+							.Text( LOCTEXT("ToleranceText", "Tolerance") )
+							.ToolTipText(this, &SDataprepFloatFilter::GetSelectedCriteriaTooltipText)
+							.Justification( ETextJustify::Center )
+						]
+					]
+					+ SHorizontalBox::Slot()
+					.Padding( 5.f )
+					[
+						SNew( SSpinBox<float> )
+						.Value( this, &SDataprepFloatFilter::GetTolerance )
+						.OnValueChanged( this, &SDataprepFloatFilter::OnToleranceChanged )
+						.OnValueCommitted( this, &SDataprepFloatFilter::OnToleranceComitted )
+						.Justification( ETextJustify::Center )
+						.MinValue( 0.f )
+						.MaxValue( TOptional< float >() )
+						.ContextMenuExtender( this, &SDataprepFloatFilter::ExtendContextMenuForToleranceBox )
+					]
 				]
 			]
 		]
 	];
+
+	if ( MatchingCriteriaParameterizationActionData && MatchingCriteriaParameterizationActionData->IsValid() )
+	{
+		if ( MatchingCriteriaParameterizationActionData->DataprepAsset->IsObjectPropertyBinded( Filter, MatchingCriteriaParameterizationActionData->PropertyChain ) )
+		{
+			MatchingCriteriaHorizontalBox->AddSlot()
+				.HAlign( HAlign_Right )
+				.VAlign( VAlign_Center )
+				.Padding( FMargin(5.f, 0.f, 0.f, 0.f) )
+				.AutoWidth()
+				[
+					SNew( SDataprepParameterizationLinkIcon, MatchingCriteriaParameterizationActionData->DataprepAsset, Filter, MatchingCriteriaParameterizationActionData->PropertyChain )
+				];
+		}
+	}
+
+	if ( EqualValueParameterizationActionData && EqualValueParameterizationActionData->IsValid() )
+	{
+		if ( EqualValueParameterizationActionData->DataprepAsset->IsObjectPropertyBinded( Filter, EqualValueParameterizationActionData->PropertyChain ) )
+		{
+			EqualValueHorizontalBox->AddSlot()
+				.HAlign( HAlign_Right )
+				.VAlign( VAlign_Center )
+				.Padding( FMargin(5.f, 0.f, 0.f, 0.f) )
+				.AutoWidth()
+				[
+					SNew( SDataprepParameterizationLinkIcon, EqualValueParameterizationActionData->DataprepAsset, Filter, EqualValueParameterizationActionData->PropertyChain )
+				];
+		}
+	}
+
+	if ( ToleranceParameterizationActionData && ToleranceParameterizationActionData->IsValid() )
+	{
+		if ( ToleranceParameterizationActionData->DataprepAsset->IsObjectPropertyBinded( Filter, ToleranceParameterizationActionData->PropertyChain ) )
+		{
+			ToleranceHorizontalBox->AddSlot()
+				.HAlign( HAlign_Right )
+				.VAlign( VAlign_Center )
+				.Padding( FMargin(5.f, 0.f, 0.f, 0.f) )
+				.AutoWidth()
+				[
+					SNew( SDataprepParameterizationLinkIcon, ToleranceParameterizationActionData->DataprepAsset, Filter, ToleranceParameterizationActionData->PropertyChain )
+				];
+		}
+	}
 }
 
 TSharedRef<SWidget> SDataprepFloatFilter::OnGenerateWidgetForMatchingCriteria(TSharedPtr<FListEntry> ListEntry) const
@@ -127,6 +274,9 @@ void SDataprepFloatFilter::OnSelectedCriteriaChanged(TSharedPtr<FListEntry> List
 	{
 		FScopedTransaction Transaction( LOCTEXT("SelectionCriteriaChangedTransaction", "Changed the Float Selection Criteria") );
 		Filter->SetFloatMatchingCriteria( FloatMatchType );
+
+		SDataprepFloatFilterUtils::PostEditChainProperty( Filter, TEXT("FloatMatchingCriteria") );
+
 		FDataprepEditorUtils::NotifySystemOfChangeInPipeline( Filter );
 	}
 }
@@ -151,6 +301,11 @@ void SDataprepFloatFilter::OnCriteriaComboBoxOpenning()
 	FloatMatchingCriteriaWidget->SetSelectedItem( ItemToSelect );
 }
 
+TSharedPtr<SWidget> SDataprepFloatFilter::OnGetContextMenuForMatchingCriteria()
+{
+	return FDataprepEditorUtils::MakeContextMenu( MatchingCriteriaParameterizationActionData );
+}
+
 void SDataprepFloatFilter::AddReferencedObjects(FReferenceCollector& Collector)
 {
 	Collector.AddReferencedObject( Filter );
@@ -178,11 +333,27 @@ void SDataprepFloatFilter::OnEqualValueComitted(float NewEqualValue, ETextCommit
 		Filter->SetEqualValue( OldEqualValue );
 		FScopedTransaction Transaction( LOCTEXT("EqualValueChangedTransaction","Change the Equal Value") );
 		Filter->SetEqualValue( NewEqualValue );
+
+		SDataprepFloatFilterUtils::PostEditChainProperty( Filter, TEXT("EqualValue") );
+
 		FDataprepEditorUtils::NotifySystemOfChangeInPipeline( Filter );
 		OldEqualValue = NewEqualValue;
 	}
 }
 
+TSharedPtr<SWidget> SDataprepFloatFilter::OnGetContextMenuForEqualValue()
+{
+	return FDataprepEditorUtils::MakeContextMenu( EqualValueParameterizationActionData );
+}
+
+void SDataprepFloatFilter::ExtendContextMenuForEqualValueBox(FMenuBuilder& MenuBuilder)
+{
+	if ( EqualValueParameterizationActionData && EqualValueParameterizationActionData->IsValid() && Filter )
+	{
+		FDataprepEditorUtils::PopulateMenuForParameterization( MenuBuilder, *EqualValueParameterizationActionData->DataprepAsset, *Filter, EqualValueParameterizationActionData->PropertyChain );
+	}
+}
+
 EVisibility SDataprepFloatFilter::GetToleranceRowVisibility() const
 {
 	check( Filter );
@@ -211,9 +382,31 @@ void SDataprepFloatFilter::OnToleranceComitted(float NewTolerance, ETextCommit::
 		Filter->SetTolerance( OldTolerance );
 		FScopedTransaction Transaction( LOCTEXT("ToleranceChangedTransaction", "Change the Tolerance") );
 		Filter->SetTolerance( NewTolerance );
+		SDataprepFloatFilterUtils::PostEditChainProperty( Filter, TEXT("Tolerance") );
 		FDataprepEditorUtils::NotifySystemOfChangeInPipeline( Filter );
 		OldTolerance = NewTolerance;
 	}
 }
 
+TSharedPtr<SWidget> SDataprepFloatFilter::OnGetContextMenuForTolerance()
+{
+	return FDataprepEditorUtils::MakeContextMenu( ToleranceParameterizationActionData );
+}
+
+void SDataprepFloatFilter::ExtendContextMenuForToleranceBox(FMenuBuilder& MenuBuilder)
+{
+	if ( ToleranceParameterizationActionData && ToleranceParameterizationActionData->IsValid() && Filter )
+	{
+		FDataprepEditorUtils::PopulateMenuForParameterization( MenuBuilder, *ToleranceParameterizationActionData->DataprepAsset, *Filter, ToleranceParameterizationActionData->PropertyChain );
+	}
+}
+
+void SDataprepFloatFilter::OnParameterizationStatusForObjectsChanged(const TSet<UObject*>* Objects)
+{
+	if ( !Objects || Objects->Contains( Filter ) )
+	{
+		UpdateVisualDisplay();
+	}
+}
+
 #undef LOCTEXT_NAMESPACE
diff --git a/Engine/Plugins/Experimental/DataPrepEditor/Source/DataPrepEditor/Private/Widgets/Action/SDataprepFloatFilter.h b/Engine/Plugins/Experimental/DataPrepEditor/Source/DataPrepEditor/Private/Widgets/Action/SDataprepFloatFilter.h
--- a/Engine/Plugins/Experimental/DataPrepEditor/Source/DataPrepEditor/Private/Widgets/Action/SDataprepFloatFilter.h
+++ b/Engine/Plugins/Experimental/DataPrepEditor/Source/DataPrepEditor/Private/Widgets/Action/SDataprepFloatFilter.h
@@ -8,12 +8,15 @@
 #include "Widgets/DeclarativeSyntaxSupport.h"
 #include "Widgets/SCompoundWidget.h"
 
+class FMenuBuilder;
 class SComboButton;
 class UDataprepFloatFilter;
 
 template<class t>
 class SComboBox;
 
+struct FDataprepParametrizationActionData;
+
 class SDataprepFloatFilter : public SCompoundWidget, public FGCObject
 {
 	SLATE_BEGIN_ARGS(SDataprepFloatFilter) {}
@@ -21,7 +24,13 @@ class SDataprepFloatFilter : public SCompoundWidget, public FGCObject
 
 	void Construct(const FArguments& InArgs, UDataprepFloatFilter& InFilter);
 
+	virtual ~SDataprepFloatFilter();
+
 private:
+
+	// Refresh the ui
+	void UpdateVisualDisplay();
+
 	// The string matching option for the combo box (Displayed text, Tooltip, mapping for the UEnum)
 	using FListEntry = TTuple<FText, FText, int32>;
 
@@ -31,17 +40,25 @@ private:
 	FText GetSelectedCriteriaTooltipText() const;
 	void OnSelectedCriteriaChanged(TSharedPtr<FListEntry> ListEntry, ESelectInfo::Type SelectionType);
 	void OnCriteriaComboBoxOpenning();
+	TSharedPtr<SWidget> OnGetContextMenuForMatchingCriteria();
 
 	// Those functions are for the float that will be compare against the fetched float
 	float GetEqualValue() const;
 	void OnEqualValueChanged(float NewEqualValue);
 	void OnEqualValueComitted(float NewEqualValue, ETextCommit::Type CommitType);
+	TSharedPtr<SWidget> OnGetContextMenuForEqualValue();
+	void ExtendContextMenuForEqualValueBox(FMenuBuilder& MenuBuilder);
 
 	// Those function are for the tolerance display
 	EVisibility GetToleranceRowVisibility() const;
 	float GetTolerance() const;
 	void OnToleranceChanged(float NewTolerance);
 	void OnToleranceComitted(float NewTolerance, ETextCommit::Type CommitType);
+	TSharedPtr<SWidget> OnGetContextMenuForTolerance();
+	void ExtendContextMenuForToleranceBox(FMenuBuilder& MenuBuilder);
+
+	// Event for dataprep parameterization
+	void OnParameterizationStatusForObjectsChanged(const TSet<UObject*>* Objects);
 
 	//~ FGCObject interface
 	virtual void AddReferencedObjects(FReferenceCollector& Collector) override;
@@ -55,4 +72,10 @@ private:
 	TArray<TSharedPtr<FListEntry>> FloatMatchingOptions;
 
 	TSharedPtr<SComboBox<TSharedPtr<FListEntry>>> FloatMatchingCriteriaWidget;
+
+	TSharedPtr<FDataprepParametrizationActionData> MatchingCriteriaParameterizationActionData;
+	TSharedPtr<FDataprepParametrizationActionData> EqualValueParameterizationActionData;
+	TSharedPtr<FDataprepParametrizationActionData> ToleranceParameterizationActionData;
+
+	FDelegateHandle OnParameterizationStatusForObjectsChangedHandle;
 };
diff --git a/Engine/Plugins/Experimental/DataPrepEditor/Source/DataPrepEditor/Private/Widgets/Action/SDataprepStringFilter.cpp b/Engine/Plugins/Experimental/DataPrepEditor/Source/DataPrepEditor/Private/Widgets/Action/SDataprepStringFilter.cpp
--- a/Engine/Plugins/Experimental/DataPrepEditor/Source/DataPrepEditor/Private/Widgets/Action/SDataprepStringFilter.cpp
+++ b/Engine/Plugins/Experimental/DataPrepEditor/Source/DataPrepEditor/Private/Widgets/Action/SDataprepStringFilter.cpp
@@ -2,15 +2,22 @@
 
 #include "Widgets/Action/SDataprepStringFilter.h"
 
+#include "DataPrepAsset.h"
+#include "DataprepCoreUtils.h"
 #include "DataprepEditorUtils.h"
+#include "Parameterization/DataprepParameterizationUtils.h"
 #include "SchemaActions/DataprepFetcherMenuActionCollector.h"
 #include "SchemaActions/DataprepSchemaAction.h"
 #include "SelectionSystem/DataprepStringFilter.h"
 #include "SelectionSystem/DataprepStringsArrayFilter.h"
 #include "Widgets/Action/DataprepActionWidgetsUtils.h"
 #include "Widgets/Action/SDataprepFetcherSelector.h"
+#include "Widgets/DataprepWidgets.h"
+#include "Widgets/Parameterization/SDataprepParameterizationLinkIcon.h"
 
+#include "Framework/MultiBox/MultiBoxBuilder.h"
 #include "Internationalization/Text.h"
+#include "Misc/AssertionMacros.h"
 #include "ScopedTransaction.h"
 #include "Widgets/Input/SComboBox.h"
 #include "Widgets/Input/SComboButton.h"
@@ -28,6 +35,51 @@ void SDataprepStringFilter<FilterType>::Construct(const FArguments& InArgs, Filt
 
 	DataprepActionWidgetsUtils::GenerateListEntriesFromEnum< EDataprepStringMatchType >( StringMatchingOptions );
 
+	if ( UDataprepAsset* DataprepAsset = FDataprepCoreUtils::GetDataprepAssetOfObject( &InFilter ) )
+	{
+		UClass* FilterClass = InFilter.GetClass();
+
+		{
+			FName PropertyName = FName( TEXT("StringMatchingCriteria") );
+			UProperty* Property = FilterClass->FindPropertyByName( PropertyName );
+			check( Property );
+			TArray<FDataprepPropertyLink> PropertyChain;
+			PropertyChain.Emplace( Property, PropertyName, INDEX_NONE );
+
+			MatchingCriteriaParameterizationActionData = MakeShared<FDataprepParametrizationActionData>( *DataprepAsset, InFilter, PropertyChain );
+		}
+
+		{
+			FName PropertyName = FName( TEXT("UserString") );
+			UProperty* Property = FilterClass->FindPropertyByName( PropertyName );
+			check( Property );
+			TArray<FDataprepPropertyLink> PropertyChain;
+			PropertyChain.Emplace( Property, PropertyName, INDEX_NONE );
+
+			UserStringParameterizationActionData = MakeShared<FDataprepParametrizationActionData>( *DataprepAsset, InFilter, PropertyChain);
+		}
+
+		OnParameterizationStatusForObjectsChangedHandle = DataprepAsset->OnParameterizationStatusForObjectsChanged.AddSP( this, &SDataprepStringFilter<FilterType>::OnParameterizationStatusForObjectsChanged );
+	}
+
+	UpdateVisualDisplay();
+}
+
+template <class FilterType>
+SDataprepStringFilter<FilterType>::~SDataprepStringFilter()
+{
+	if ( UDataprepAsset* DataprepAsset = FDataprepCoreUtils::GetDataprepAssetOfObject( Filter ) )
+	{
+		DataprepAsset->OnParameterizationStatusForObjectsChanged.Remove( OnParameterizationStatusForObjectsChangedHandle );
+	}
+}
+
+template <class FilterType>
+void SDataprepStringFilter<FilterType>::UpdateVisualDisplay()
+{
+	TSharedPtr<SHorizontalBox> MatchingCriteriaHorizontalBox;
+	TSharedPtr<SHorizontalBox> UserStringHorizontalBox;
+
 	ChildSlot
 	[
 		SNew( SBox )
@@ -37,30 +89,77 @@ void SDataprepStringFilter<FilterType>::Construct(const FArguments& InArgs, Filt
 			+ SHorizontalBox::Slot()
 			.Padding( 5.f )
 			[
-				SAssignNew( StringMatchingCriteriaWidget, SComboBox< TSharedPtr< FListEntry > > )
-				.OptionsSource( &StringMatchingOptions )
-				.OnGenerateWidget( this, &SDataprepStringFilter::OnGenerateWidgetForMatchingCriteria )
-				.OnSelectionChanged( this, &SDataprepStringFilter::OnSelectedCriteriaChanged )
-				.OnComboBoxOpening( this, &SDataprepStringFilter::OnCriteriaComboBoxOpenning )
+				SNew( SDataprepContextMenuOverride )
+				.OnContextMenuOpening( this, &SDataprepStringFilter<FilterType>::OnGetContextMenuForMatchingCriteria )
 				[
-					SNew( STextBlock )
-					.Text( this, &SDataprepStringFilter::GetSelectedCriteriaText )
-					.ToolTipText( this, &SDataprepStringFilter::GetSelectedCriteriaTooltipText )
-					.Justification( ETextJustify::Center )
+					SAssignNew( MatchingCriteriaHorizontalBox, SHorizontalBox )
+					+ SHorizontalBox::Slot()
+					[
+						SAssignNew( StringMatchingCriteriaWidget, SComboBox< TSharedPtr< FListEntry > > )
+						.OptionsSource( &StringMatchingOptions )
+						.OnGenerateWidget( this, &SDataprepStringFilter::OnGenerateWidgetForMatchingCriteria )
+						.OnSelectionChanged( this, &SDataprepStringFilter::OnSelectedCriteriaChanged )
+						.OnComboBoxOpening( this, &SDataprepStringFilter::OnCriteriaComboBoxOpenning )
+						[
+							SNew( STextBlock )
+							.Text( this, &SDataprepStringFilter::GetSelectedCriteriaText )
+							.ToolTipText( this, &SDataprepStringFilter::GetSelectedCriteriaTooltipText )
+							.Justification( ETextJustify::Center )
+						]
+					]
 				]
 			
 			]
 			+ SHorizontalBox::Slot()
 			.Padding( 5.f )
 			[
-				SNew( SEditableTextBox )
-				.Text( this, &SDataprepStringFilter::GetUserString )
-				.OnTextChanged( this, &SDataprepStringFilter::OnUserStringChanged )
-				.OnTextCommitted( this, &SDataprepStringFilter::OnUserStringComitted )
-				.Justification( ETextJustify::Center )
+				SNew( SDataprepContextMenuOverride )
+				.OnContextMenuOpening( this, &SDataprepStringFilter<FilterType>::OnGetContextMenuForUserString )
+				[
+					SAssignNew( UserStringHorizontalBox, SHorizontalBox )
+					+ SHorizontalBox::Slot()
+					[
+						SNew( SEditableTextBox )
+						.Text( this, &SDataprepStringFilter::GetUserString )
+						.ContextMenuExtender( this, &SDataprepStringFilter::ExtendContextMenuForUserStringBox )
+						.OnTextChanged( this, &SDataprepStringFilter::OnUserStringChanged )
+						.OnTextCommitted( this, &SDataprepStringFilter::OnUserStringComitted )
+						.Justification( ETextJustify::Center )
+					]
+				]
 			]
 		]
 	];
+
+	if ( MatchingCriteriaParameterizationActionData && MatchingCriteriaParameterizationActionData->IsValid() )
+	{
+		if ( MatchingCriteriaParameterizationActionData->DataprepAsset->IsObjectPropertyBinded( Filter, MatchingCriteriaParameterizationActionData->PropertyChain ) )
+		{
+			MatchingCriteriaHorizontalBox->AddSlot()
+				.HAlign( HAlign_Right )
+				.VAlign( VAlign_Center )
+				.Padding( FMargin(5.f, 0.f, 0.f, 0.f) )
+				.AutoWidth()
+				[
+					SNew( SDataprepParameterizationLinkIcon, MatchingCriteriaParameterizationActionData->DataprepAsset, Filter, MatchingCriteriaParameterizationActionData->PropertyChain )
+				];
+		}
+	}
+
+	if ( UserStringParameterizationActionData && UserStringParameterizationActionData->IsValid() )
+	{
+		if ( UserStringParameterizationActionData->DataprepAsset->IsObjectPropertyBinded( Filter, UserStringParameterizationActionData->PropertyChain ) )
+		{
+			UserStringHorizontalBox->AddSlot()
+				.HAlign( HAlign_Right )
+				.VAlign( VAlign_Center )
+				.Padding( FMargin(5.f, 0.f, 0.f, 0.f) )
+				.AutoWidth()
+				[
+					SNew( SDataprepParameterizationLinkIcon, UserStringParameterizationActionData->DataprepAsset, Filter, UserStringParameterizationActionData->PropertyChain )
+				];
+		}
+	}
 }
 
 template <class FilterType>
@@ -108,6 +207,13 @@ void SDataprepStringFilter<FilterType>::OnCriteriaComboBoxOpenning()
 	StringMatchingCriteriaWidget->SetSelectedItem( ItemToSelect );
 }
 
+template <class FilterType>
+TSharedPtr<SWidget> SDataprepStringFilter<FilterType>::OnGetContextMenuForMatchingCriteria()
+{
+	return FDataprepEditorUtils::MakeContextMenu( MatchingCriteriaParameterizationActionData );
+}
+
+
 template <class FilterType>
 void SDataprepStringFilter<FilterType>::OnSelectedCriteriaChanged(TSharedPtr<FListEntry> ListEntry, ESelectInfo::Type SelectionType)
 {
@@ -119,6 +225,17 @@ void SDataprepStringFilter<FilterType>::OnSelectedCriteriaChanged(TSharedPtr<FLi
 	{	
 		FScopedTransaction Transaction( LOCTEXT("SelectionCriteriaChangedTransaction","Changed the String Selection Criteria") );
 		Filter->SetStringMatchingCriteria( StringMatchType );
+
+		UProperty* Property = Filter->GetClass()->FindPropertyByName( TEXT("StringMatchingCriteria") );
+		check( Property );
+
+		FEditPropertyChain EditChain;
+		EditChain.AddHead( Property );
+		EditChain.SetActivePropertyNode( Property );
+		FPropertyChangedEvent EditPropertyChangeEvent( Property, EPropertyChangeType::ValueSet );
+		FPropertyChangedChainEvent EditChangeChainEvent( EditChain, EditPropertyChangeEvent );
+		Filter->PostEditChangeChainProperty( EditChangeChainEvent );
+
 		FDataprepEditorUtils::NotifySystemOfChangeInPipeline( Filter );
 	}
 }
@@ -137,6 +254,13 @@ void SDataprepStringFilter<FilterType>::OnUserStringChanged(const FText& NewText
 	Filter->SetUserString( NewText.ToString() );
 }
 
+template <class FilterType>
+void SDataprepStringFilter<FilterType>::ExtendContextMenuForUserStringBox(FMenuBuilder& MenuBuilder)
+{
+	FDataprepEditorUtils::PopulateMenuForParameterization( MenuBuilder, *UserStringParameterizationActionData->DataprepAsset,
+		*Filter, UserStringParameterizationActionData->PropertyChain );
+}
+
 template <class FilterType>
 void SDataprepStringFilter<FilterType>::OnUserStringComitted(const FText& NewText, ETextCommit::Type CommitType)
 {
@@ -147,11 +271,37 @@ void SDataprepStringFilter<FilterType>::OnUserStringComitted(const FText& NewTex
 		Filter->SetUserString( OldUserString );
 		FScopedTransaction Transaction( LOCTEXT("SelectionStringChangedTransaction","Changed the Selection String") );
 		Filter->SetUserString( NewUserString );
+
+		UProperty* Property = Filter->GetClass()->FindPropertyByName( TEXT("UserString") );
+		check( Property );
+
+		FEditPropertyChain EditChain;
+		EditChain.AddHead( Property );
+		EditChain.SetActivePropertyNode( Property );
+		FPropertyChangedEvent EditPropertyChangeEvent( Property, EPropertyChangeType::ValueSet );
+		FPropertyChangedChainEvent EditChangeChainEvent( EditChain, EditPropertyChangeEvent );
+		Filter->PostEditChangeChainProperty( EditChangeChainEvent );
+
 		FDataprepEditorUtils::NotifySystemOfChangeInPipeline( Filter );
 		OldUserString = NewUserString;
 	}
 }
 
+template <class FilterType>
+TSharedPtr<SWidget> SDataprepStringFilter<FilterType>::OnGetContextMenuForUserString()
+{
+	return FDataprepEditorUtils::MakeContextMenu( UserStringParameterizationActionData );
+}
+
+template <class FilterType>
+void SDataprepStringFilter<FilterType>::OnParameterizationStatusForObjectsChanged(const TSet<UObject*>* Objects)
+{
+	if ( !Objects || Objects->Contains( Filter ) )
+	{
+		UpdateVisualDisplay();
+	}
+}
+
 template <class FilterType>
 void SDataprepStringFilter<FilterType>::AddReferencedObjects(FReferenceCollector& Collector)
 {
diff --git a/Engine/Plugins/Experimental/DataPrepEditor/Source/DataPrepEditor/Private/Widgets/Action/SDataprepStringFilter.h b/Engine/Plugins/Experimental/DataPrepEditor/Source/DataPrepEditor/Private/Widgets/Action/SDataprepStringFilter.h
--- a/Engine/Plugins/Experimental/DataPrepEditor/Source/DataPrepEditor/Private/Widgets/Action/SDataprepStringFilter.h
+++ b/Engine/Plugins/Experimental/DataPrepEditor/Source/DataPrepEditor/Private/Widgets/Action/SDataprepStringFilter.h
@@ -3,16 +3,20 @@
 #pragma once
 
 #include "CoreMinimal.h"
+#include "Delegates/IDelegateInstance.h"
 #include "Types/SlateEnums.h"
 #include "UObject/GCObject.h"
 #include "Widgets/DeclarativeSyntaxSupport.h"
 #include "Widgets/SCompoundWidget.h"
 
+class FMenuBuilder;
 class SComboButton;
 
 template<class t>
 class SComboBox;
 
+struct FDataprepParametrizationActionData;
+
 template <class FilterType>
 class SDataprepStringFilter : public SCompoundWidget,  public FGCObject
 {
@@ -21,7 +25,12 @@ class SDataprepStringFilter : public SCompoundWidget,  public FGCObject
 
 	void Construct(const FArguments& InArgs, FilterType& InFilter);
 
+	virtual ~SDataprepStringFilter();
+
 private:
+
+	void UpdateVisualDisplay();
+
 	// The string matching option for the combo box ( Displayed text, Tooltip, mapping for the UEnum)
 	using FListEntry = TTuple<FText, FText, int32>;
 
@@ -31,11 +40,16 @@ private:
 	FText GetSelectedCriteriaTooltipText() const;
 	void OnSelectedCriteriaChanged(TSharedPtr<FListEntry> ListEntry, ESelectInfo::Type SelectionType);
 	void OnCriteriaComboBoxOpenning();
+	TSharedPtr<SWidget> OnGetContextMenuForMatchingCriteria();
 
 	// This function is for the string that will be compare against the fetched string
 	FText GetUserString() const;
 	void OnUserStringChanged(const FText& NewText);
 	void OnUserStringComitted(const FText& NewText, ETextCommit::Type CommitType);
+	void ExtendContextMenuForUserStringBox(FMenuBuilder& MenuBuilder);
+	TSharedPtr<SWidget> OnGetContextMenuForUserString();
+
+	void OnParameterizationStatusForObjectsChanged(const TSet<UObject*>* Objects);
 
 	//~ FGCObject interface
 	virtual void AddReferencedObjects(FReferenceCollector& Collector) override;
@@ -48,4 +62,9 @@ private:
 	TArray<TSharedPtr<FListEntry>> StringMatchingOptions;
 
 	TSharedPtr<SComboBox<TSharedPtr<FListEntry>>> StringMatchingCriteriaWidget;
+
+	TSharedPtr<FDataprepParametrizationActionData> MatchingCriteriaParameterizationActionData;
+	TSharedPtr<FDataprepParametrizationActionData> UserStringParameterizationActionData;
+
+	FDelegateHandle OnParameterizationStatusForObjectsChangedHandle;
 };
diff --git a/Engine/Plugins/Experimental/DataPrepEditor/Source/DataPrepEditor/Private/Widgets/DataprepWidgets.cpp b/Engine/Plugins/Experimental/DataPrepEditor/Source/DataPrepEditor/Private/Widgets/DataprepWidgets.cpp
--- a/Engine/Plugins/Experimental/DataPrepEditor/Source/DataPrepEditor/Private/Widgets/DataprepWidgets.cpp
+++ b/Engine/Plugins/Experimental/DataPrepEditor/Source/DataPrepEditor/Private/Widgets/DataprepWidgets.cpp
@@ -9,6 +9,7 @@
 #include "DataprepEditorUtils.h"
 #include "DataprepParameterizableObject.h"
 #include "Parameterization/DataprepParameterizationUtils.h"
+#include "Widgets/Parameterization/SDataprepParameterizationLinkIcon.h"
 
 #include "ContentBrowserModule.h"
 #include "DetailLayoutBuilder.h"
@@ -17,11 +18,17 @@
 #include "EditorFontGlyphs.h"
 #include "EditorFontGlyphs.h"
 #include "EditorStyleSet.h"
+#include "Framework/Application/MenuStack.h"
+#include "Framework/Application/SlateApplication.h"
+#include "Framework/MultiBox/MultiBoxBuilder.h"
 #include "Framework/MultiBox/MultiBoxBuilder.h"
 #include "HAL/PlatformProcess.h"
 #include "IContentBrowserSingleton.h"
 #include "IDetailTreeNode.h"
 #include "IPropertyRowGenerator.h"
+#include "Input/Events.h"
+#include "Layout/Geometry.h"
+#include "Layout/WidgetPath.h"
 #include "Misc/Paths.h"
 #include "Modules/ModuleManager.h"
 #include "PropertyCustomizationHelpers.h"
@@ -38,6 +45,7 @@
 #include "Widgets/Layout/SBox.h"
 #include "Widgets/Layout/SGridPanel.h"
 #include "Widgets/Layout/SSpacer.h"
+#include "Widgets/SNullWidget.h"
 #include "Widgets/Text/SInlineEditableTextBlock.h"
 #include "Widgets/Text/STextBlock.h"
 
@@ -418,62 +426,7 @@ TSharedRef< SWidget > SDataprepDetailsView::CreateDefaultWidget( TSharedPtr< SWi
 	TSharedRef<SHorizontalBox> NameColumn = SNew(SHorizontalBox)
 		.Clipping(EWidgetClipping::OnDemand);
 
-	// Optionally add the parameterization widget
-	if ( DataprepAssetForParameterization.IsValid() )
-	{
-		if ( ParameterizationContext.State == EParametrizationState::CanBeParameterized )
-		{
-			NameColumn->AddSlot()
-				.HAlign(HAlign_Left)
-				.VAlign(VAlign_Center)
-				.AutoWidth()
-				[
-
-					SNew(STextBlock)
-					.Font(FEditorStyle::Get().GetFontStyle("FontAwesome.11"))
-					.Text(FEditorFontGlyphs::Plus_Circle)
-					.Margin(FMargin(5.0f, 5.0f, 0.0f, 0.0f))
-					.ColorAndOpacity(FColor::Blue)
-					.OnDoubleClicked_Lambda([this, PropertyChain = ParameterizationContext.PropertyChain](const FGeometry&,const FPointerEvent&) -> FReply
-						{
-							if ( UDataprepAsset* DataprepAsset = DataprepAssetForParameterization.Get() )
-							{
-								FScopedTransaction Transaction( LOCTEXT("DataprepBindingToParameterization","Adding Parameter") );
-								DataprepAssetForParameterization->BindObjectPropertyToParameterization( DetailedObjectAsParameterizable, PropertyChain, *PropertyChain.Last().CachedProperty->GetDisplayNameText().ToString() );
-								bRefreshObjectToDisplay = true;
-							}
-
-							return FReply::Handled();
-						})
-				];
-		}
-		else if ( ParameterizationContext.State == EParametrizationState::IsParameterized )
-		{
-			NameColumn->AddSlot()
-				.HAlign(HAlign_Left)
-				.VAlign(VAlign_Center)
-				.AutoWidth()
-				[
-
-					SNew(STextBlock)
-					.Font(FEditorStyle::Get().GetFontStyle("FontAwesome.11"))
-					.Text(FEditorFontGlyphs::Minus_Circle)
-					.Margin(FMargin(5.0f, 5.0f, 0.0f, 0.0f))
-					.ColorAndOpacity(FColor::Red)
-					.OnDoubleClicked_Lambda([this, PropertyChain = ParameterizationContext.PropertyChain](const FGeometry&, const FPointerEvent&)->FReply
-					{
-						if (UDataprepAsset * DataprepAsset = DataprepAssetForParameterization.Get())
-						{
-							FScopedTransaction Transaction( LOCTEXT("DataprepRemoveBindingFromParameterization","Removing Parameter") );
-							DataprepAssetForParameterization->RemoveObjectPropertyFromParameterization( DetailedObjectAsParameterizable, PropertyChain );
-							bRefreshObjectToDisplay = true;
-						}
-
-						return FReply::Handled();
-					})
-				];
-		}
-	}
+	NameWidget->SetClipping( EWidgetClipping::OnDemand );
 
 	// Add the name widget
 	NameColumn->AddSlot()
@@ -484,44 +437,78 @@ TSharedRef< SWidget > SDataprepDetailsView::CreateDefaultWidget( TSharedPtr< SWi
 			NameWidget.ToSharedRef()
 		];
 
+	if ( ParameterizationContext.State == EParametrizationState::IsParameterized )
+	{
+		NameColumn->AddSlot()
+			.HAlign(HAlign_Right)
+			.VAlign(VAlign_Center)
+			.Padding(FMargin(5.f, 0.f, 5.f, 0.f))
+			.AutoWidth()
+			[
+				SNew(SDataprepParameterizationLinkIcon, DataprepAssetForParameterization.Get(), DetailedObjectAsParameterizable, ParameterizationContext.PropertyChain)
+			];
+	}
+
+	FOnContextMenuOpening OnContextMenuOpening;
 
-	return SNew(SHorizontalBox)
-	+ SHorizontalBox::Slot()
-	.FillWidth(1.f)
-	.VAlign(VAlign_Center)
-	.HAlign(HAlign_Left)
-	[
-		SNew(SSplitter)
-		.Style(FEditorStyle::Get(), "DetailsView.Splitter")
-		.PhysicalSplitterHandleSize(1.0f)
-		.HitDetectionSplitterHandleSize(5.0f)
-		.ResizeMode( ESplitterResizeMode::Fill )
-		+ SSplitter::Slot()
-		.Value(ColumnSizeData->LeftColumnWidth)
-		.OnSlotResized( SSplitter::FOnSlotResized::CreateSP( this, &SDataprepDetailsView::OnLeftColumnResized ) )
-		[
-			NameColumn
-		]
-		+ SSplitter::Slot()
-		.Value(ColumnSizeData->RightColumnWidth)
-		.OnSlotResized(ColumnSizeData->OnWidthChanged)
+	if ( ParameterizationContext.State == EParametrizationState::IsParameterized || ParameterizationContext.State == EParametrizationState::CanBeParameterized )
+	{
+		OnContextMenuOpening.BindLambda( [this, InPropertyChain = ParameterizationContext.PropertyChain] () -> TSharedPtr<SWidget>
+			{
+				if ( UDataprepAsset * DataprepAsset = DataprepAssetForParameterization.Get() )
+				{
+					if ( DetailedObjectAsParameterizable )
+					{
+						FMenuBuilder MenuBuilder(true, nullptr);
+						FDataprepEditorUtils::PopulateMenuForParameterization(MenuBuilder, *DataprepAsset, *DetailedObjectAsParameterizable, InPropertyChain);
+						return MenuBuilder.MakeWidget();
+					}
+				}
+				return TSharedPtr<SWidget>();
+			});
+	}
+
+	return SNew(SDataprepContextMenuOverride)
+		.OnContextMenuOpening(OnContextMenuOpening)
 		[
 			SNew(SHorizontalBox)
-			.Clipping(EWidgetClipping::OnDemand)
 			+ SHorizontalBox::Slot()
-			.Padding( 5.0f, 2.5f, 2.0f, 2.5f )
-			.HAlign( HAlign )
-			.VAlign( VAlign )
+			.FillWidth(1.f)
+			.VAlign(VAlign_Center)
+			.HAlign(HAlign_Left)
 			[
-				// Trick to force the splitter widget to fill up the space of its parent
-				// Strongly inspired from SDetailSingleItemRow
-				SNew(SConstrainedBox)
+				SNew(SSplitter)
+				.Style(FEditorStyle::Get(), "DetailsView.Splitter")
+				.PhysicalSplitterHandleSize(1.0f)
+				.HitDetectionSplitterHandleSize(5.0f)
+				.ResizeMode( ESplitterResizeMode::Fill )
+				+ SSplitter::Slot()
+				.Value(ColumnSizeData->LeftColumnWidth)
+				.OnSlotResized( SSplitter::FOnSlotResized::CreateSP( this, &SDataprepDetailsView::OnLeftColumnResized ) )
 				[
-					ValueWidget.ToSharedRef()
+					NameColumn
+				]
+				+ SSplitter::Slot()
+				.Value(ColumnSizeData->RightColumnWidth)
+				.OnSlotResized(ColumnSizeData->OnWidthChanged)
+				[
+					SNew(SHorizontalBox)
+					.Clipping(EWidgetClipping::OnDemand)
+					+ SHorizontalBox::Slot()
+					.Padding( 5.0f, 2.5f, 2.0f, 2.5f )
+					.HAlign( HAlign )
+					.VAlign( VAlign )
+					[
+						// Trick to force the splitter widget to fill up the space of its parent
+						// Strongly inspired from SDetailSingleItemRow
+						SNew(SConstrainedBox)
+						[
+							ValueWidget.ToSharedRef()
+						]
+					]
 				]
 			]
-		]
-	];
+		];
 }
 
 void SDataprepDetailsView::OnPropertyChanged(const FPropertyChangedEvent& InEvent)
@@ -530,7 +517,7 @@ void SDataprepDetailsView::OnPropertyChanged(const FPropertyChangedEvent& InEven
 
 	if( TrackedProperties.Contains( InEvent.Property ) )
 	{
-		bRefreshObjectToDisplay = true;
+		ForceRefresh();
 	}
 }
 
@@ -544,7 +531,7 @@ void SDataprepDetailsView::OnObjectReplaced(const TMap<UObject*, UObject*>& Repl
 			DetailedObjectAsParameterizable = static_cast<UDataprepParameterizableObject*>( DetailedObject );
 		}
 
-		bRefreshObjectToDisplay = true;
+		ForceRefresh();
 	}
 }
 
@@ -553,6 +540,14 @@ void SDataprepDetailsView::ForceRefresh()
 	bRefreshObjectToDisplay = true;
 }
 
+void SDataprepDetailsView::OnDataprepParameterizationStatusForObjectsChanged(const TSet<UObject*>* Objects)
+{
+	if ( !Objects || Objects->Contains( DetailedObjectAsParameterizable ) )
+	{
+		ForceRefresh();
+	}
+}
+
 void SDataprepDetailsView::AddWidgets( const TArray< TSharedRef< IDetailTreeNode > >& DetailTree, TSharedPtr<SGridPanel>& GridPanel, int32& Index, float LeftPadding, const FDataprepParameterizationContext& InParameterizationContext)
 {
 	auto IsDetailNodeDisplayable = []( const TSharedPtr< IPropertyHandle >& PropertyHandle)
@@ -769,7 +764,7 @@ void SDataprepDetailsView::Construct()
 		UDataprepAsset* DataprepAsset = FDataprepParameterizationUtils::GetDataprepAssetForParameterization( DetailedObject );
 		if ( DataprepAsset )
 		{
-			OnDataprepParameterizationWasModifiedHandle = DataprepAsset->OnParameterizationWasModified.AddSP( this, &SDataprepDetailsView::ForceRefresh );
+			OnDataprepParameterizationStatusForObjectsChangedHandle = DataprepAsset->OnParameterizationStatusForObjectsChanged.AddSP( this, &SDataprepDetailsView::OnDataprepParameterizationStatusForObjectsChanged );
 		}
 
 		if ( DetailedObject->IsA<UDataprepParameterizableObject>() )
@@ -835,7 +830,7 @@ SDataprepDetailsView::~SDataprepDetailsView()
 
 	if ( UDataprepAsset* DataprepAsset = DataprepAssetForParameterization.Get() )
 	{
-		DataprepAsset->OnParameterizationWasModified.Remove( OnDataprepParameterizationWasModifiedHandle );
+		DataprepAsset->OnParameterizationStatusForObjectsChanged.Remove( OnDataprepParameterizationStatusForObjectsChangedHandle );
 	}
 }
 
@@ -864,7 +859,7 @@ void SDataprepDetailsView::SetObjectToDisplay(UObject& Object)
 		{
 			DetailedObjectAsParameterizable = static_cast<UDataprepParameterizableObject*>( DetailedObject );
 		}
-		bRefreshObjectToDisplay = true;
+		ForceRefresh();
 	}
 }
 
@@ -875,4 +870,38 @@ void SDataprepDetailsView::AddReferencedObjects(FReferenceCollector& Collector)
 	Collector.AddReferencedObjects( TrackedProperties );
 }
 
+
+void SDataprepContextMenuOverride::Construct(const FArguments& InArgs)
+{
+	OnContextMenuOpening = InArgs._OnContextMenuOpening;
+	
+	ChildSlot
+	[
+		InArgs._DefaultSlot.Widget
+	];
+}
+
+FReply SDataprepContextMenuOverride::OnMouseButtonUp(const FGeometry& MyGeometry, const FPointerEvent& MouseEvent)
+{
+	if ( MouseEvent.GetEffectingButton() == EKeys::RightMouseButton && OnContextMenuOpening.IsBound() )
+	{
+		TSharedPtr<SWidget> ContextMenu = OnContextMenuOpening.Execute();
+		if ( ContextMenu )
+		{ 
+			FWidgetPath WidgetPath = MouseEvent.GetEventPath() ? *MouseEvent.GetEventPath() : FWidgetPath();
+			FSlateApplication::Get().PushMenu(
+				AsShared(),
+				WidgetPath,
+				ContextMenu.ToSharedRef(),
+				MouseEvent.GetScreenSpacePosition(),
+				FPopupTransitionEffect(FPopupTransitionEffect::ContextMenu)
+				);
+
+			return FReply::Handled();
+		}
+	}
+
+	return FReply::Unhandled();
+}
+
 #undef LOCTEXT_NAMESPACE
diff --git a/Engine/Plugins/Experimental/DataPrepEditor/Source/DataPrepEditor/Private/Widgets/DataprepWidgets.h b/Engine/Plugins/Experimental/DataPrepEditor/Source/DataPrepEditor/Private/Widgets/DataprepWidgets.h
--- a/Engine/Plugins/Experimental/DataPrepEditor/Source/DataPrepEditor/Private/Widgets/DataprepWidgets.h
+++ b/Engine/Plugins/Experimental/DataPrepEditor/Source/DataPrepEditor/Private/Widgets/DataprepWidgets.h
@@ -4,7 +4,10 @@
 
 #include "CoreMinimal.h"
 
+#include "Parameterization/DataprepParameterizationUtils.h"
+
 #include "Framework/Application/SlateApplication.h"
+#include "Framework/SlateDelegates.h"
 #include "Widgets/Layout/SSplitter.h"
 #include "Widgets/SCompoundWidget.h"
 
@@ -13,9 +16,6 @@ class UDataprepAsset;
 class UDataprepContentConsumer;
 class UDataprepParameterizableObject;
 
-struct FDataprepParameterizationContext;
-struct FDataprepPropertyLink;
-
 struct FDataprepDetailsViewColumnSizeData
 {
 	TAttribute<float> LeftColumnWidth;
@@ -156,6 +156,8 @@ private:
 
 	void ForceRefresh();
 
+	void OnDataprepParameterizationStatusForObjectsChanged(const TSet<UObject*>* Objects);
+
 private:
 	/** Row generator applied on detailed object */
 	TSharedPtr< class IPropertyRowGenerator > Generator;
@@ -187,5 +189,24 @@ private:
 	// This pointer to a dataprep asset is used by the parameterization system. It should be nullptr when the parameterization shouldn't be shown
 	TWeakObjectPtr<UDataprepAsset> DataprepAssetForParameterization;
 
-	FDelegateHandle OnDataprepParameterizationWasModifiedHandle;
+	FDelegateHandle OnDataprepParameterizationStatusForObjectsChangedHandle;
+};
+
+/**
+ * This is widget simply exist to open contextual menu
+ */
+class SDataprepContextMenuOverride : public SCompoundWidget
+{
+public:
+	SLATE_BEGIN_ARGS(SDataprepContextMenuOverride) {}
+		SLATE_DEFAULT_SLOT(FArguments, DefaultSlot)
+		SLATE_EVENT(FOnContextMenuOpening, OnContextMenuOpening)
+	SLATE_END_ARGS()
+
+	void Construct(const FArguments& InArgs);
+
+	virtual FReply OnMouseButtonUp(const FGeometry& MyGeometry, const FPointerEvent& MouseEvent) override;
+
+private:
+	FOnContextMenuOpening OnContextMenuOpening;
 };
diff --git a/Engine/Plugins/Experimental/DataPrepEditor/Source/DataPrepEditor/Private/Widgets/Parameterization/SDataprepLinkToParameter.cpp b/Engine/Plugins/Experimental/DataPrepEditor/Source/DataPrepEditor/Private/Widgets/Parameterization/SDataprepLinkToParameter.cpp
new file mode 100644
--- /dev/null
+++ b/Engine/Plugins/Experimental/DataPrepEditor/Source/DataPrepEditor/Private/Widgets/Parameterization/SDataprepLinkToParameter.cpp
@@ -0,0 +1,226 @@
+// Copyright 1998-2019 Epic Games, Inc. All Rights Reserved.
+
+#include "Widgets/Parameterization/SDataprepLinkToParameter.h"
+
+#include "DataPrepAsset.h"
+#include "DataprepEditorUtils.h"
+#include "DataprepParameterizableObject.h"
+
+#include "Async/ParallelFor.h"
+#include "EditorStyleSet.h"
+#include "Framework/Application/SlateApplication.h"
+#include "ScopedTransaction.h"
+#include "Widgets/Input/SEditableTextBox.h"
+#include "Widgets/Layout/SBorder.h"
+#include "Widgets/Layout/SBox.h"
+#include "Widgets/SBoxPanel.h"
+#include "Widgets/Text/STextBlock.h"
+#include "Widgets/Views/SListView.h"
+#include "Widgets/Views/STableRow.h"
+ 
+#define LOCTEXT_NAMESPACE "SDataprepLinkToParameter"
+
+FText SDataprepLinkToParameter::EmptyNameErrorText = LOCTEXT("EmptyNameErrorText", "The parameter name can't be an empty string");
+FText SDataprepLinkToParameter::InvalidNameErrorText = LOCTEXT("InvalidNameErrorText", "The parameter name is already used for another data type");
+
+FText SDataprepLinkToParameter::LinkToParameterTransactionText = LOCTEXT("LinkedParameterToParameterization", "Linked property to parameterization");
+
+void SDataprepLinkToParameter::Construct(const FArguments& InArgs, const TSharedRef<FDataprepParametrizationActionData>& InParameterizationActionData)
+{
+	ParameterizationActionData = InParameterizationActionData;
+
+	
+	if ( ParameterizationActionData->IsValid() )
+	{
+		if ( UProperty* Property = ParameterizationActionData->PropertyChain.Last().CachedProperty.Get() )
+		{
+			UClass* PropertyClass = Property->GetClass();
+			ParameterizationActionData->DataprepAsset->GetExistingParameterNamesForType( PropertyClass, ValidExistingNames, InvalidNames );
+
+			FString SuggestedName = ParameterizationActionData->PropertyChain.Last().PropertyName.ToString();
+
+			bool bIsSuggestedNameValid = !SuggestedName.IsEmpty() && !InvalidNames.Contains( SuggestedName );
+			bool bIsSuggestedNameAParameter = ValidExistingNames.Contains( SuggestedName );
+
+			ShownItems.Reserve( ValidExistingNames.Num() + 1 );
+			CreateNewOption = MakeShared<FString>( LOCTEXT("CreateNewParamerterLabel", "Create New Parameter").ToString() );
+			if ( bIsSuggestedNameValid && !bIsSuggestedNameAParameter )
+			{
+				ShownItems.Add( CreateNewOption );
+			}
+
+			UnfilteredSuggestions.Reserve( ValidExistingNames.Num() );
+			for ( const FString& String : ValidExistingNames )
+			{
+				UnfilteredSuggestions.Add( MakeShared<FString>( String ) );
+			}
+
+			ShownItems.Append( UnfilteredSuggestions );
+
+
+			ChildSlot
+			[
+				SNew( SBox )
+				.MinDesiredWidth( 400 )
+				.MaxDesiredHeight( 400 )
+				[
+					SNew( SVerticalBox )
+					+ SVerticalBox::Slot()
+					.HAlign( HAlign_Fill )
+					.VAlign( VAlign_Top )
+					.AutoHeight()
+					[
+						SAssignNew( TextBox, SEditableTextBox )
+						.Text( FText::FromString( SuggestedName ) )
+						.HintText( LOCTEXT("ParameterNameHintText", "Enter a parameter name") )
+						.OnTextChanged( this, &SDataprepLinkToParameter::OnTextChanged )
+						.OnTextCommitted( this, &SDataprepLinkToParameter::OnTextCommited )
+					]
+					+ SVerticalBox::Slot()
+					.HAlign( HAlign_Fill )
+					.VAlign( VAlign_Top )
+					.Padding( FMargin( 0.f, 5.f, 0.f, 0.f ) )
+					[
+						SAssignNew( SuggestionList, SListView<TSharedPtr<FString>> )
+						.ListItemsSource( &ShownItems )
+						.OnGenerateRow( this, &SDataprepLinkToParameter::OnGenerateRowForList )
+						.SelectionMode( ESelectionMode::Single )
+						.OnSelectionChanged( this, &SDataprepLinkToParameter::OnSelectionChanged )
+					]
+				]
+			];
+
+			if ( !bIsSuggestedNameValid )
+			{
+				if ( SuggestedName.IsEmpty() )
+				{
+					TextBox->SetError( EmptyNameErrorText.ToString() );
+				}
+				else
+				{
+					TextBox->SetError( InvalidNameErrorText.ToString() );
+				}
+			}
+		}
+	}
+}
+
+TSharedRef<class ITableRow> SDataprepLinkToParameter::OnGenerateRowForList(TSharedPtr<FString> Item, const TSharedRef<class STableViewBase>& OwnerTable)
+{
+	if ( Item == CreateNewOption )
+	{
+		return SNew( STableRow<TSharedPtr<FString>>, OwnerTable )
+			[
+				SNew( STextBlock )
+				.Text( FText::FromString(*Item) )
+				.TextStyle( &FCoreStyle::Get(), ISlateStyle::Join("ToolBar", ".Label") )
+				.Margin( FMargin(2.f) )
+			];
+	}
+
+	return SNew( STableRow<TSharedPtr<FString>>, OwnerTable )
+		[
+			SNew( STextBlock )
+			.Text( FText ::FromString( *Item ) )
+			.TextStyle( &FCoreStyle::Get(), ISlateStyle::Join( "ToolBar", ".Label" ) )
+			.HighlightText_Lambda( [this]() { return FText::FromString( ParameterName ); } )
+			.Margin( FMargin( 2.f ) )
+		];
+}
+
+void SDataprepLinkToParameter::OnSelectionChanged(TSharedPtr<FString> InItem, ESelectInfo::Type SelectionType)
+{
+	FScopedTransaction Transaction( LinkToParameterTransactionText );
+	if ( InItem == CreateNewOption )
+	{
+		ParameterName = TextBox->GetText().ToString();
+		ParameterizationActionData->DataprepAsset->BindObjectPropertyToParameterization( ParameterizationActionData->Object, ParameterizationActionData->PropertyChain, FName( *ParameterName ) );
+	}
+	else
+	{
+		ParameterizationActionData->DataprepAsset->BindObjectPropertyToParameterization( ParameterizationActionData->Object, ParameterizationActionData->PropertyChain, FName( *( *InItem.Get() ) ) );
+	}
+
+	FSlateApplication::Get().DismissAllMenus();
+}
+
+void SDataprepLinkToParameter::OnTextChanged(const FText& Text)
+{
+	ParameterName = Text.ToString();
+	bool bHasError = SetErrorMessage();
+
+	// Update the shown item list
+	
+	ShownItems.Empty( ValidExistingNames.Num() + 1 );
+
+	if ( !bHasError && !ValidExistingNames.Contains( ParameterName ) )
+	{
+		ShownItems.Add( CreateNewOption );
+	}
+
+	TArray<bool> FilterResult;
+	FilterResult.AddZeroed( UnfilteredSuggestions.Num() );
+
+	ParallelFor( UnfilteredSuggestions.Num(), [this, &FilterResult](int32 Index)
+		{
+			FilterResult[Index] = UnfilteredSuggestions[Index]->Contains( ParameterName );
+		});
+
+	for ( int32 Index = 0; Index < FilterResult.Num(); Index++ )
+	{
+		if ( FilterResult[Index] )
+		{
+			ShownItems.Add( UnfilteredSuggestions[Index] );
+		}
+	}
+
+	SuggestionList->RequestListRefresh();
+}
+
+void SDataprepLinkToParameter::OnTextCommited(const FText& Text, ETextCommit::Type CommitType)
+{
+
+	ParameterName = Text.ToString();
+	bool bHasError = SetErrorMessage();
+	
+	if ( CommitType == ETextCommit::OnEnter )
+	{
+		if ( !bHasError &&  ParameterizationActionData->IsValid() )
+		{
+			FScopedTransaction Transaction( LinkToParameterTransactionText );
+			ParameterizationActionData->DataprepAsset->BindObjectPropertyToParameterization( ParameterizationActionData->Object, ParameterizationActionData->PropertyChain, FName( *ParameterName ) );
+		}
+
+		FSlateApplication::Get().DismissAllMenus();
+	}
+}
+
+void SDataprepLinkToParameter::Tick(const FGeometry& AllottedGeometry, const double InCurrentTime, const float InDeltaTime)
+{
+	if ( !bHadFirstTick )
+	{
+		FSlateApplication::Get().SetKeyboardFocus( TextBox, EFocusCause::SetDirectly );
+		TextBox->SelectAllText();
+		bHadFirstTick = true;
+	}
+}
+
+bool SDataprepLinkToParameter::SetErrorMessage()
+{
+	if ( ParameterName.IsEmpty() )
+	{
+		TextBox->SetError( EmptyNameErrorText.ToString() );
+		return true;
+	}
+
+	if ( InvalidNames.Contains( ParameterName ) )
+	{
+		TextBox->SetError( InvalidNameErrorText.ToString() );
+		return true;
+	}
+
+	TextBox->SetError( FString() );
+	return false;
+}
+
+#undef LOCTEXT_NAMESPACE
diff --git a/Engine/Plugins/Experimental/DataPrepEditor/Source/DataPrepEditor/Private/Widgets/Parameterization/SDataprepLinkToParameter.h b/Engine/Plugins/Experimental/DataPrepEditor/Source/DataPrepEditor/Private/Widgets/Parameterization/SDataprepLinkToParameter.h
new file mode 100644
--- /dev/null
+++ b/Engine/Plugins/Experimental/DataPrepEditor/Source/DataPrepEditor/Private/Widgets/Parameterization/SDataprepLinkToParameter.h
@@ -0,0 +1,69 @@
+// Copyright 1998-2019 Epic Games, Inc. All Rights Reserved.
+
+#pragma once
+
+#include "Parameterization/DataprepParameterizationUtils.h"
+
+#include "CoreMinimal.h"
+#include "Types/SlateEnums.h"
+#include "UObject/GCObject.h"
+#include "Widgets/DeclarativeSyntaxSupport.h"
+#include "Widgets/SCompoundWidget.h"
+
+class UDataprepAsset;
+class UDataprepParameterizableObject;
+
+struct FDataprepParametrizationActionData;
+
+class SDataprepLinkToParameter : public SCompoundWidget
+{
+public:
+	SLATE_BEGIN_ARGS(SDataprepLinkToParameter) {}
+	SLATE_END_ARGS();
+
+	void Construct(const FArguments& InArgs, const TSharedRef<FDataprepParametrizationActionData>& InParameterizationActionData);
+
+private:
+
+	TSharedRef<class ITableRow> OnGenerateRowForList(TSharedPtr<FString> Item, const TSharedRef<class STableViewBase>& OwnerTable);
+
+	void OnSelectionChanged(TSharedPtr<FString> InItem, ESelectInfo::Type SelectionType);
+
+	void OnTextChanged(const FText& Text);
+
+	void OnTextCommited(const FText& Text, ETextCommit::Type CommitType);
+
+	virtual void Tick(const FGeometry& AllottedGeometry, const double InCurrentTime, const float InDeltaTime) override;
+
+	/** Return true if there is a error */
+	bool SetErrorMessage();
+
+	// The information that will be used to create the binding
+	TSharedPtr<FDataprepParametrizationActionData> ParameterizationActionData;
+
+	TSharedPtr<FString> CreateNewOption;
+	TArray<TSharedPtr<FString>> UnfilteredSuggestions;
+
+	// The items currently shown
+	TArray<TSharedPtr<FString>> ShownItems;
+
+	// The names that can be use because they are used for other property types
+	TSet<FString> InvalidNames;
+
+	TSet<FString> ValidExistingNames;
+
+	TSharedPtr<class SEditableTextBox> TextBox;
+	TSharedPtr<class SListView<TSharedPtr<FString>>> SuggestionList;
+
+	// The current parameter name entered by the user
+	FString ParameterName;
+
+	static FText EmptyNameErrorText;
+	static FText InvalidNameErrorText;
+
+	static FText LinkToParameterTransactionText;
+
+	// Minor hack to set the focus on the text box
+	bool bHadFirstTick = false;
+};
+
diff --git a/Engine/Plugins/Experimental/DataPrepEditor/Source/DataPrepEditor/Private/Widgets/Parameterization/SDataprepParameterizationLinkIcon.cpp b/Engine/Plugins/Experimental/DataPrepEditor/Source/DataPrepEditor/Private/Widgets/Parameterization/SDataprepParameterizationLinkIcon.cpp
new file mode 100644
--- /dev/null
+++ b/Engine/Plugins/Experimental/DataPrepEditor/Source/DataPrepEditor/Private/Widgets/Parameterization/SDataprepParameterizationLinkIcon.cpp
@@ -0,0 +1,36 @@
+// Copyright 1998-2019 Epic Games, Inc. All Rights Reserved.
+
+#include "Widgets/Parameterization/SDataprepParameterizationLinkIcon.h"
+
+#include "DataPrepAsset.h"
+#include "DataprepParameterizableObject.h"
+#include "Parameterization/DataprepParameterizationUtils.h"
+
+#include "EditorFontGlyphs.h"
+#include "EditorStyleSet.h"
+#include "Internationalization/Text.h"
+
+#define LOCTEXT_NAMESPACE "SDataprepParameterizationLinkIcon"
+
+void SDataprepParameterizationLinkIcon::Construct(const FArguments& InArgs, UDataprepAsset* DataprepAsset, UDataprepParameterizableObject* Object, const TArray<FDataprepPropertyLink>& PropertyChain)
+{
+	if ( DataprepAsset )
+	{
+		FName ParameterName = DataprepAsset->GetNameOfParameterForObjectProperty( Object, PropertyChain );
+
+		if ( !ParameterName.IsNone() )
+		{
+			ChildSlot
+			[
+					SNew( STextBlock )
+					.Font( FEditorStyle::Get().GetFontStyle("FontAwesome.11") )
+					.Text( FEditorFontGlyphs::Link )
+					.Justification( ETextJustify::Center )
+					.ColorAndOpacity( FColor( 170, 238, 160 ) )
+					.ToolTipText( FText::Format( LOCTEXT("LinkIconTooltipText", "This property is linked to the parameter {0}"), FText::FromName( ParameterName ) ) )
+			];
+		}
+	}
+}
+
+#undef LOCTEXT_NAMESPACE
diff --git a/Engine/Plugins/Experimental/DataPrepEditor/Source/DataPrepEditor/Private/Widgets/Parameterization/SDataprepParameterizationLinkIcon.h b/Engine/Plugins/Experimental/DataPrepEditor/Source/DataPrepEditor/Private/Widgets/Parameterization/SDataprepParameterizationLinkIcon.h
new file mode 100644
--- /dev/null
+++ b/Engine/Plugins/Experimental/DataPrepEditor/Source/DataPrepEditor/Private/Widgets/Parameterization/SDataprepParameterizationLinkIcon.h
@@ -0,0 +1,24 @@
+// Copyright 1998-2019 Epic Games, Inc. All Rights Reserved.
+
+#pragma once
+
+#include "CoreMinimal.h"
+#include "Widgets/DeclarativeSyntaxSupport.h"
+#include "Widgets/SCompoundWidget.h"
+
+class UDataprepAsset;
+class UDataprepParameterizableObject;
+
+struct FDataprepPropertyLink;
+
+/**
+ * This is widget serves to display some information about a binded property to the dataprep parameterization
+ */
+class SDataprepParameterizationLinkIcon : public SCompoundWidget
+{
+public:
+	SLATE_BEGIN_ARGS(SDataprepParameterizationLinkIcon) {}
+	SLATE_END_ARGS();
+
+	void Construct(const FArguments& InArgs, UDataprepAsset* DataprepAsset, UDataprepParameterizableObject* Object, const TArray<FDataprepPropertyLink>& PropertyChain);
+};
diff --git a/Engine/Source/Runtime/Slate/Private/Widgets/Input/SEditableTextBox.cpp b/Engine/Source/Runtime/Slate/Private/Widgets/Input/SEditableTextBox.cpp
--- a/Engine/Source/Runtime/Slate/Private/Widgets/Input/SEditableTextBox.cpp
+++ b/Engine/Source/Runtime/Slate/Private/Widgets/Input/SEditableTextBox.cpp
@@ -64,6 +64,7 @@ void SEditableTextBox::Construct( const FArguments& InArgs )
 					.Justification( InArgs._Justification )
 					.AllowContextMenu( InArgs._AllowContextMenu )
 					.OnContextMenuOpening( InArgs._OnContextMenuOpening )
+					.ContextMenuExtender( InArgs._ContextMenuExtender )
 					.OnTextChanged( InArgs._OnTextChanged )
 					.OnTextCommitted( InArgs._OnTextCommitted )
 					.MinDesiredWidth( InArgs._MinDesiredWidth )
diff --git a/Engine/Source/Runtime/Slate/Public/Widgets/Input/SEditableTextBox.h b/Engine/Source/Runtime/Slate/Public/Widgets/Input/SEditableTextBox.h
--- a/Engine/Source/Runtime/Slate/Public/Widgets/Input/SEditableTextBox.h
+++ b/Engine/Source/Runtime/Slate/Public/Widgets/Input/SEditableTextBox.h
@@ -3,16 +3,17 @@
 #pragma once
 
 #include "CoreMinimal.h"
-#include "Misc/Attribute.h"
-#include "Styling/SlateColor.h"
 #include "Fonts/SlateFontInfo.h"
+#include "Framework/MultiBox/MultiBoxExtender.h"
 #include "Input/Reply.h"
-#include "Widgets/DeclarativeSyntaxSupport.h"
 #include "Layout/Margin.h"
-#include "Styling/SlateTypes.h"
+#include "Misc/Attribute.h"
 #include "Styling/CoreStyle.h"
-#include "Widgets/Layout/SBorder.h"
+#include "Styling/SlateColor.h"
+#include "Styling/SlateTypes.h"
+#include "Widgets/DeclarativeSyntaxSupport.h"
 #include "Widgets/Input/SEditableText.h"
+#include "Widgets/Layout/SBorder.h"
 
 class IErrorReportingWidget;
 class SBox;
@@ -103,6 +104,9 @@ public:
 		/** Delegate to call before a context menu is opened. User returns the menu content or null to the disable context menu */
 		SLATE_EVENT(FOnContextMenuOpening, OnContextMenuOpening)
 
+		/** Menu extender for the right-click context menu */
+		SLATE_EVENT(FMenuExtensionDelegate, ContextMenuExtender)
+
 		/** Called whenever the text is changed programmatically or interactively by the user */
 		SLATE_EVENT( FOnTextChanged, OnTextChanged )
 
diff --git a/Engine/Source/Runtime/Slate/Public/Widgets/Input/SSpinBox.h b/Engine/Source/Runtime/Slate/Public/Widgets/Input/SSpinBox.h
--- a/Engine/Source/Runtime/Slate/Public/Widgets/Input/SSpinBox.h
+++ b/Engine/Source/Runtime/Slate/Public/Widgets/Input/SSpinBox.h
@@ -3,27 +3,28 @@
 #pragma once
 
 #include "CoreMinimal.h"
-#include "Misc/Attribute.h"
-#include "InputCoreTypes.h"
-#include "Layout/Margin.h"
 #include "Fonts/SlateFontInfo.h"
-#include "Layout/Visibility.h"
-#include "Widgets/DeclarativeSyntaxSupport.h"
-#include "Styling/SlateColor.h"
+#include "Framework/Application/SlateApplication.h"
+#include "Framework/MultiBox/MultiBoxExtender.h"
 #include "Input/CursorReply.h"
 #include "Input/Events.h"
 #include "Input/Reply.h"
-#include "Widgets/SCompoundWidget.h"
+#include "InputCoreTypes.h"
+#include "Layout/Margin.h"
+#include "Layout/Visibility.h"
+#include "Misc/Attribute.h"
+#include "Rendering/DrawElements.h"
 #include "Styling/CoreStyle.h"
+#include "Styling/SlateColor.h"
+#include "Styling/SlateTypes.h"
+#include "Templates/IsIntegral.h"
+#include "Widgets/DeclarativeSyntaxSupport.h"
+#include "Widgets/Images/SImage.h"
 #include "Widgets/Input/NumericTypeInterface.h"
+#include "Widgets/Input/SEditableText.h"
 #include "Widgets/SBoxPanel.h"
-#include "Styling/SlateTypes.h"
+#include "Widgets/SCompoundWidget.h"
 #include "Widgets/Text/STextBlock.h"
-#include "Widgets/Input/SEditableText.h"
-#include "Rendering/DrawElements.h"
-#include "Framework/Application/SlateApplication.h"
-#include "Widgets/Images/SImage.h"
-#include "Templates/IsIntegral.h"
 
 /*
  * This function compute a slider position by simulating two log on both side of the neutral value
@@ -139,6 +140,8 @@ public:
 		SLATE_ARGUMENT( TSharedPtr< INumericTypeInterface<NumericType> >, TypeInterface )
 		/** If refresh requests for the viewport should happen for all value changes **/
 		SLATE_ARGUMENT(bool, PreventThrottling)
+		/** Menu extender for the right-click context menu */
+		SLATE_EVENT( FMenuExtensionDelegate, ContextMenuExtender )
 
 	SLATE_END_ARGS()
 
@@ -257,6 +260,7 @@ public:
 				.VirtualKeyboardType(EKeyboardType::Keyboard_Number)
 				.Justification(InArgs._Justification)
 				.VirtualKeyboardTrigger(EVirtualKeyboardTrigger::OnAllFocusEvents)
+				.ContextMenuExtender( InArgs._ContextMenuExtender )
 			]
 
 			+SHorizontalBox::Slot()
