Dataprep asset preview icons are the same for all asset types #jira UE-82127 #rb JeanLuc.Corenthin

diff --git a/Engine/Plugins/Experimental/DataPrepEditor/Source/DataPrepEditor/Private/Widgets/SAssetsPreviewWidget.cpp b/Engine/Plugins/Experimental/DataPrepEditor/Source/DataPrepEditor/Private/Widgets/SAssetsPreviewWidget.cpp
--- a/Engine/Plugins/Experimental/DataPrepEditor/Source/DataPrepEditor/Private/Widgets/SAssetsPreviewWidget.cpp
+++ b/Engine/Plugins/Experimental/DataPrepEditor/Source/DataPrepEditor/Private/Widgets/SAssetsPreviewWidget.cpp
@@ -2,6 +2,8 @@
 
 #include "SAssetsPreviewWidget.h"
 
+#include "AssetToolsModule.h"
+#include "IAssetTypeActions.h"
 #include "EditorStyleSet.h"
 #include "Widgets/Images/SImage.h"
 #include "Widgets/Text/STextBlock.h"
@@ -120,6 +122,19 @@ namespace AssetPreviewWidget
 			PreviewWidgetWeakPtr = InPreviewWidget;
 			ItemWeakPtr = InItem;
 
+			FSlateColor IconColor(FLinearColor::White);
+			TWeakObjectPtr< UObject > AssetPtr = InItem->AssetPtr;
+			if (AssetPtr.IsValid())
+			{
+				static FAssetToolsModule& AssetToolsModule = FModuleManager::LoadModuleChecked<FAssetToolsModule>(TEXT("AssetTools"));
+				TSharedPtr<IAssetTypeActions> AssetTypeActions = AssetToolsModule.Get().GetAssetTypeActionsForClass(AssetPtr->GetClass()).Pin();
+
+				if (AssetTypeActions.IsValid())
+				{
+					IconColor = FSlateColor(AssetTypeActions->GetTypeColor());
+				}
+			}
+
 			STableRow::Construct(
 				STableRow::FArguments()
 				.Style(FEditorStyle::Get(), "ContentBrowser.AssetListView.TableRow")
@@ -136,7 +151,7 @@ namespace AssetPreviewWidget
 					// Item icon
 					SNew(SImage)
 					.Image(this, &SAssetPreviewTableRow::GetIconBrush)
-					.ColorAndOpacity(FSlateColor(FLinearColor::White))
+					.ColorAndOpacity(IconColor)
 				]
 
 				+ SHorizontalBox::Slot()
