Implement a merge operation based on the LODProxy merge capability #jira UEENT-3015 #rb JeanLuc.Corenthin

diff --git a/Engine/Plugins/Experimental/DataPrepEditor/Source/DataPrepLibraries/Private/DataprepEditingOperations.cpp b/Engine/Plugins/Experimental/DataPrepEditor/Source/DataPrepLibraries/Private/DataprepEditingOperations.cpp
--- a/Engine/Plugins/Experimental/DataPrepEditor/Source/DataPrepLibraries/Private/DataprepEditingOperations.cpp
+++ b/Engine/Plugins/Experimental/DataPrepEditor/Source/DataPrepLibraries/Private/DataprepEditingOperations.cpp
@@ -68,6 +68,7 @@ namespace DatasmithEditingOperationsUtils
 	void FindActorsToMerge(const TArray<AActor*>& ChildrenActors, TArray<AActor*>& ActorsToMerge);
 	void FindActorsToCollapseOrDelete(const TArray<AActor*>& ActorsToVisit, TArray<AActor*>& ActorsToCollapse, TArray<UObject*>& ActorsToDelete );
 	void GetRootActors(UWorld* World, TArray<AActor*>& RootActors);
+	void GetActorsToMerge(UWorld*& World, const TArray<UObject*>& InObjects, TArray<AActor*>& ActorsToMerge, TArray<UPrimitiveComponent*> ComponentsToMerge);
 
 	int32 GetActorDepth(AActor* Actor)
 	{
@@ -161,45 +162,7 @@ void UDataprepMergeActorsOperation::OnExecution_Implementation(const FDataprepCo
 	TArray<UPrimitiveComponent*> ComponentsToMerge;
 	UWorld* CurrentWorld = nullptr;
 
-	for(UObject* Object : InContext.Objects)
-	{
-		if( AActor* Actor = Cast<AActor>(Object) )
-		{
-			if(!Actor->IsPendingKillOrUnreachable())
-			{
-				if(CurrentWorld == nullptr)
-				{
-					CurrentWorld = Actor->GetWorld();
-				}
-
-				if(CurrentWorld != Actor->GetWorld())
-				{
-					// #ueent_todo: Warn that incompatible actor found and discarded
-					continue;
-				}
-
-				TInlineComponentArray<UStaticMeshComponent*> ComponentArray;
-				Actor->GetComponents<UStaticMeshComponent>(ComponentArray);
-
-				bool bMeshActorIsValid = false;
-				for (UStaticMeshComponent* MeshComponent : ComponentArray)
-				{
-					if (MeshComponent->GetStaticMesh())
-					{
-						bMeshActorIsValid = true;
-						ComponentsToMerge.Add(MeshComponent);
-					}
-				}
-
-				//Actor needs at least one StaticMeshComponent to be considered valid
-				if (bMeshActorIsValid)
-				{
-					ActorsToMerge.Add( Actor );
-				}
-
-			}
-		}
-	}
+	DatasmithEditingOperationsUtils::GetActorsToMerge(CurrentWorld, InContext.Objects, ActorsToMerge, ComponentsToMerge);
 
 	// Nothing to do if there is none or only one static mesh actor
 	if( ActorsToMerge.Num() < 2 && ComponentsToMerge.Num() < 2)
@@ -325,6 +288,121 @@ bool UDataprepMergeActorsOperation::MergeStaticMeshActors(UWorld* World, const T
 	return true;
 }
 
+void UDataprepCreateProxyMeshOperation::OnExecution_Implementation(const FDataprepContext& InContext)
+{
+	TArray<AActor*> ActorsToMerge;
+	TArray<UPrimitiveComponent*> ComponentsToMerge;
+	UWorld* CurrentWorld = nullptr;
+
+	DatasmithEditingOperationsUtils::GetActorsToMerge(CurrentWorld, InContext.Objects, ActorsToMerge, ComponentsToMerge);
+
+	// Nothing to do if there is none or only one static mesh actor
+	if(ActorsToMerge.Num() < 2 && ComponentsToMerge.Num() < 2)
+	{
+		UE_LOG(LogDataprep, Log, TEXT("No static mesh actors to merge"));
+		return;
+	}
+
+	// Validate render data for static meshes
+	TSet<UStaticMesh*> StaticMeshes;
+	for(UPrimitiveComponent* PrimitiveComponent : ComponentsToMerge)
+	{
+		if(UStaticMeshComponent* StaticMeshComponent = Cast<UStaticMeshComponent>(PrimitiveComponent))
+		{
+			if(StaticMeshComponent->GetStaticMesh()->RenderData == nullptr)
+			{
+				StaticMeshes.Add(StaticMeshComponent->GetStaticMesh());
+			}
+		}
+	}
+
+	DataprepOperationsLibraryUtil::FStaticMeshBuilder StaticMeshBuilder( StaticMeshes );
+
+	FMeshProxySettings ProxySettings;
+	ProxySettings.ScreenSize = ProxySettings.ScreenSize * ((100.0f - ReductionPercent) / 100.f);
+
+	// Generate proxy mesh and proxy material assets 
+	FCreateProxyDelegate ProxyDelegate;
+	ProxyDelegate.BindLambda( [&](const FGuid Guid, TArray<UObject*>& AssetsToSync )
+	{
+		UStaticMesh* ProxyMesh = nullptr;
+		if(!AssetsToSync.FindItemByClass(&ProxyMesh))
+		{
+			UE_LOG(LogDataprep, Error, TEXT("CreateProxyMesh failed. No mesh was created."));
+			return;
+		}
+
+		// Add asset to set of assets in Dataprep action working set
+		MergedMesh = Cast<UStaticMesh>(AddAsset(ProxyMesh, NewActorLabel.IsEmpty() ? TEXT("Proxy_Mesh") : *NewActorLabel));
+		if(!MergedMesh)
+		{
+			UE_LOG(LogDataprep, Error, TEXT("CreateProxyMesh failed. Internal error while creating the merged mesh."));
+			return;
+		}
+
+		// Place new mesh in the world (on a new actor)
+		MergedActor = Cast<AStaticMeshActor>(CreateActor(AStaticMeshActor::StaticClass(), NewActorLabel.IsEmpty() ? TEXT("Proxy_Actor") : *NewActorLabel));
+		if(!MergedActor)
+		{
+			UE_LOG(LogDataprep, Error, TEXT("CreateProxyMesh failed. Internal error while creating the merged actor."));
+			return;
+		}
+
+		MergedActor->GetStaticMeshComponent()->SetStaticMesh(MergedMesh);
+		MergedActor->SetActorLabel(NewActorLabel.IsEmpty() ? TEXT("Proxy_Actor") : *NewActorLabel);
+		CurrentWorld->UpdateCullDistanceVolumes(MergedActor, MergedActor->GetStaticMeshComponent());
+	});
+
+	FGuid JobGuid = FGuid::NewGuid();
+
+	const IMeshMergeUtilities& MergeUtilities = FModuleManager::Get().LoadModuleChecked<IMeshMergeModule>("MeshMergeUtilities").GetUtilities();
+	MergeUtilities.CreateProxyMesh(ActorsToMerge, ProxySettings, nullptr, GetTransientPackage(), FString(), JobGuid, ProxyDelegate);
+
+	// Position the merged actor at the right location
+	if(MergedActor->GetRootComponent() == nullptr)
+	{
+		USceneComponent* RootComponent = NewObject< USceneComponent >(MergedActor, USceneComponent::StaticClass(), *MergedActor->GetActorLabel(), RF_Transactional);
+
+		MergedActor->AddInstanceComponent(RootComponent);
+		MergedActor->SetRootComponent(RootComponent);
+	}
+
+	// Keep the merged actor in the hierarchy, taking the parent of the first component
+	// In the future, the merged actor could be attached to the common ancestor instead of the first parent in the list
+	MergedActor->GetRootComponent()->AttachToComponent(ComponentsToMerge[0]->GetAttachParent(), FAttachmentTransformRules::KeepWorldTransform);
+
+	// Collect all objects to be deleted
+	TArray<UObject*> ObjectsToDelete;
+	ObjectsToDelete.Reserve(ComponentsToMerge.Num() + ActorsToMerge.Num());
+
+	// Simple way to delete the actors: detach the merged components if it's safe to do so
+	for(UPrimitiveComponent* Component : ComponentsToMerge)
+	{
+		if(Component->GetNumChildrenComponents() == 0)
+		{
+			Component->DetachFromComponent(FDetachmentTransformRules::KeepRelativeTransform);
+		}
+		else if(UStaticMeshComponent* StaticMeshComponent = Cast<UStaticMeshComponent>(Component))
+		{
+			// To keep the children in the hierarchy, simply hide the StaticMesh, but the StaticMeshComponent could be replaced by a SceneComponent and
+			// then attach the children to it
+			StaticMeshComponent->SetStaticMesh(nullptr);
+		}
+	}
+
+	// Then delete the merged actors that don't have any children component
+	for(AActor* Actor : ActorsToMerge)
+	{
+		UPrimitiveComponent* RootComponent = Cast<UPrimitiveComponent>(Actor->GetRootComponent());
+		if(RootComponent && RootComponent->GetNumChildrenComponents() == 0)
+		{
+			ObjectsToDelete.Add(Actor);
+		}
+	}
+
+	DeleteObjects(ObjectsToDelete);
+}
+
 void UDataprepDeleteUnusedAssetsOperation::OnExecution_Implementation(const FDataprepContext& InContext)
 {
 	UWorld* World = nullptr;
@@ -593,6 +671,50 @@ namespace DatasmithEditingOperationsUtils
 		}
 	}
 
+	void GetActorsToMerge(UWorld*& World, const TArray<UObject*>& InObjects, TArray<AActor*>& ActorsToMerge, TArray<UPrimitiveComponent*> ComponentsToMerge)
+	{
+		World = nullptr;
+
+		for(UObject* Object : InObjects)
+		{
+			if(AActor* Actor = Cast<AActor>(Object))
+			{
+				if(!Actor->IsPendingKillOrUnreachable())
+				{
+					if(World == nullptr)
+					{
+						World = Actor->GetWorld();
+					}
+
+					if(World != Actor->GetWorld())
+					{
+						// #ueent_todo: Warn that incompatible actor found and discarded
+						continue;
+					}
+
+					TInlineComponentArray<UStaticMeshComponent*> ComponentArray;
+					Actor->GetComponents<UStaticMeshComponent>(ComponentArray);
+
+					bool bMeshActorIsValid = false;
+					for(UStaticMeshComponent* MeshComponent : ComponentArray)
+					{
+						if(MeshComponent->GetStaticMesh())
+						{
+							bMeshActorIsValid = true;
+							ComponentsToMerge.Add(MeshComponent);
+						}
+					}
+
+					//Actor needs at least one StaticMeshComponent to be considered valid
+					if(bMeshActorIsValid)
+					{
+						ActorsToMerge.Add(Actor);
+					}
+				}
+			}
+		}
+	}
+
 	FMergingData::FMergingData(const TArray<UPrimitiveComponent*>& PrimitiveComponents)
 	{
 		Data.Reserve(PrimitiveComponents.Num());
diff --git a/Engine/Plugins/Experimental/DataPrepEditor/Source/DataPrepLibraries/Private/DataprepEditingOperations.h b/Engine/Plugins/Experimental/DataPrepEditor/Source/DataPrepLibraries/Private/DataprepEditingOperations.h
--- a/Engine/Plugins/Experimental/DataPrepEditor/Source/DataPrepLibraries/Private/DataprepEditingOperations.h
+++ b/Engine/Plugins/Experimental/DataPrepEditor/Source/DataPrepLibraries/Private/DataprepEditingOperations.h
@@ -87,6 +87,41 @@ protected:
 	AStaticMeshActor* MergedActor;
 };
 
+UCLASS(Experimental, Category = ObjectOperation, Meta = (DisplayName = "Create Proxy Mesh", ToolTip = "Merge the meshes into a (simplified) proxy mesh"))
+class UDataprepCreateProxyMeshOperation : public UDataprepEditingOperation
+{
+	GENERATED_BODY()
+
+public:
+	UDataprepCreateProxyMeshOperation()
+		: ReductionPercent(50.f)
+		, MergedMesh(nullptr)
+		, MergedActor(nullptr)
+	{
+	}
+
+public:
+	//~ Begin UDataprepOperation Interface
+	virtual FText GetCategory_Implementation() const override
+	{
+		return FDataprepOperationCategories::ObjectOperation;
+	}
+
+	/** Settings to use for the create proxy operation */
+	UPROPERTY(EditAnywhere, Category = ProxySettings)
+	FString NewActorLabel;
+
+	UPROPERTY(EditAnywhere, Category = ProxySettings,  meta = (UIMin = 0, UIMax = 100))
+	float ReductionPercent;
+
+protected:
+	virtual void OnExecution_Implementation(const FDataprepContext& InContext) override;
+	//~ End UDataprepOperation Interface
+protected:
+	UStaticMesh* MergedMesh;
+	AStaticMeshActor* MergedActor;
+};
+
 UCLASS(Experimental, Category = ObjectOperation, Meta = (DisplayName="Delete Unused Assets", ToolTip = "Delete assets that are not referenced by any objects") )
 class UDataprepDeleteUnusedAssetsOperation : public UDataprepEditingOperation
 {
