Fixed the build issues in Non Unity.

diff --git a/Engine/Plugins/Experimental/DataPrepEditor/Source/DataPrepEditor/Private/DataprepEditorUtils.h b/Engine/Plugins/Experimental/DataPrepEditor/Source/DataPrepEditor/Private/DataprepEditorUtils.h
--- a/Engine/Plugins/Experimental/DataPrepEditor/Source/DataPrepEditor/Private/DataprepEditorUtils.h
+++ b/Engine/Plugins/Experimental/DataPrepEditor/Source/DataPrepEditor/Private/DataprepEditorUtils.h
@@ -2,6 +2,8 @@
 
 #pragma once
 
+#include "Parameterization/DataprepParameterizationUtils.h"
+
 #include "CoreMinimal.h"
 #include "Fonts/SlateFontInfo.h"
 
@@ -10,8 +12,6 @@ class SWidget;
 class UDataprepAsset;
 class UDataprepParameterizableObject;
 
-struct FDataprepPropertyLink;
-
 struct FDataprepParametrizationActionData : public FGCObject
 {
 	FDataprepParametrizationActionData(UDataprepAsset& InDataprepAsset, UDataprepParameterizableObject& InObject, const TArray<FDataprepPropertyLink>& InPropertyChain)
diff --git a/Engine/Plugins/Experimental/DataPrepEditor/Source/DataPrepEditor/Private/Widgets/Parameterization/SDataprepParameterizationLinkIcon.cpp b/Engine/Plugins/Experimental/DataPrepEditor/Source/DataPrepEditor/Private/Widgets/Parameterization/SDataprepParameterizationLinkIcon.cpp
--- a/Engine/Plugins/Experimental/DataPrepEditor/Source/DataPrepEditor/Private/Widgets/Parameterization/SDataprepParameterizationLinkIcon.cpp
+++ b/Engine/Plugins/Experimental/DataPrepEditor/Source/DataPrepEditor/Private/Widgets/Parameterization/SDataprepParameterizationLinkIcon.cpp
@@ -9,6 +9,7 @@
 #include "EditorFontGlyphs.h"
 #include "EditorStyleSet.h"
 #include "Internationalization/Text.h"
+#include "Widgets/Text/STextBlock.h"
 
 #define LOCTEXT_NAMESPACE "SDataprepParameterizationLinkIcon"
 
