Fix non-unity
#rnx

diff --git a/Engine/Plugins/Experimental/DataPrepEditor/Source/DataPrepEditor/Private/Widgets/Parameterization/SDataprepLinkToParameter.h b/Engine/Plugins/Experimental/DataPrepEditor/Source/DataPrepEditor/Private/Widgets/Parameterization/SDataprepLinkToParameter.h
--- a/Engine/Plugins/Experimental/DataPrepEditor/Source/DataPrepEditor/Private/Widgets/Parameterization/SDataprepLinkToParameter.h
+++ b/Engine/Plugins/Experimental/DataPrepEditor/Source/DataPrepEditor/Private/Widgets/Parameterization/SDataprepLinkToParameter.h
@@ -9,6 +9,7 @@
 #include "UObject/GCObject.h"
 #include "Widgets/DeclarativeSyntaxSupport.h"
 #include "Widgets/SCompoundWidget.h"
+#include "Widgets/Views/SListView.h"
 
 class UDataprepAsset;
 class UDataprepParameterizableObject;
