Moved UI for parameterization to SDataprepAssetView

fyi julien.stjean

diff --git a/Engine/Plugins/Experimental/DataPrepEditor/Source/DataPrepEditor/Private/DataPrepEditor.cpp b/Engine/Plugins/Experimental/DataPrepEditor/Source/DataPrepEditor/Private/DataPrepEditor.cpp
--- a/Engine/Plugins/Experimental/DataPrepEditor/Source/DataPrepEditor/Private/DataPrepEditor.cpp
+++ b/Engine/Plugins/Experimental/DataPrepEditor/Source/DataPrepEditor/Private/DataPrepEditor.cpp
@@ -83,7 +83,6 @@ const FName FDataprepEditor::DetailsTabId(TEXT("DataprepEditor_Details"));
 const FName FDataprepEditor::DataprepAssetTabId(TEXT("DataprepEditor_Dataprep"));
 const FName FDataprepEditor::SceneViewportTabId(TEXT("DataprepEditor_SceneViewport"));
 const FName FDataprepEditor::DataprepStatisticsTabId(TEXT("DataprepEditor_Statistics"));
-const FName FDataprepEditor::ParameterizationDefaultTabId(TEXT("DataprepEditor_Parameterization"));
 
 static bool bLogTiming = true;
 
@@ -242,40 +241,6 @@ FDataprepEditor::~FDataprepEditor()
 	}
 }
 
-TSharedRef<SDockTab> FDataprepEditor::SpawnTabParameterization(const FSpawnTabArgs& Args)
-{
-	check(Args.GetTabId() == ParameterizationDefaultTabId);
-	FPropertyEditorModule& PropertyEditorModule = FModuleManager::Get().LoadModuleChecked<FPropertyEditorModule>("PropertyEditor");
-
-	FNotifyHook* NotifyHook = nullptr;
-
-	// Create a property view
-	FPropertyEditorModule& EditModule = FModuleManager::Get().GetModuleChecked<FPropertyEditorModule>("PropertyEditor");
-
-	FDetailsViewArgs::ENameAreaSettings NameAreaSettings = FDetailsViewArgs::ENameAreaSettings::HideNameArea;
-	FDetailsViewArgs DetailsViewArgs( /*bUpdateFromSelection=*/ false, /*bLockable=*/ false, /*bAllowSearch=*/ true, NameAreaSettings, /*bHideSelectionTip=*/ true, /*InNotifyHook=*/ NotifyHook, /*InSearchInitialKeyFocus=*/ false, NAME_None);
-	DetailsViewArgs.bShowOptions = false;
-	DetailsViewArgs.bShowPropertyMatrixButton = false;
-
-	TSharedRef<IDetailsView> DetailView = PropertyEditorModule.CreateDetailView( DetailsViewArgs );
-	if ( UObject* ParameterizationObject = DataprepAssetInterfacePtr->GetParameterizationObject() )
-	{
-		DetailView->SetObject( ParameterizationObject );
-	}
-
-	return SNew(SDockTab)
-		.Label(LOCTEXT("ParameterizationTab", "Parameterization"))
-		[
-			SNew(SBorder)
-			.Padding(2.f)
-			.BorderImage(FEditorStyle::GetBrush("ToolPanel.GroupBorder"))
-			[
-				DetailView
-			]
-		];
-}
-
-
 FName FDataprepEditor::GetToolkitFName() const
 {
 	return FName("DataprepEditor");
@@ -304,10 +269,6 @@ void FDataprepEditor::RegisterTabSpawners(const TSharedRef<class FTabManager>& I
 
 	auto WorkspaceMenuCategoryRef = WorkspaceMenuCategory.ToSharedRef();
 
-	InTabManager->RegisterTabSpawner(ParameterizationDefaultTabId, FOnSpawnTab::CreateSP(this, &FDataprepEditor::SpawnTabParameterization))
-		.SetDisplayName(LOCTEXT("ParameterizationTab", "Parameterization"))
-		.SetGroup(WorkspaceMenuCategoryRef);
-
 	InTabManager->RegisterTabSpawner(ScenePreviewTabId, FOnSpawnTab::CreateSP(this, &FDataprepEditor::SpawnTabScenePreview))
 		.SetDisplayName(LOCTEXT("ScenePreviewTab", "Scene Preview"))
 		.SetGroup(WorkspaceMenuCategoryRef)
diff --git a/Engine/Plugins/Experimental/DataPrepEditor/Source/DataPrepEditor/Private/DataPrepEditor.h b/Engine/Plugins/Experimental/DataPrepEditor/Source/DataPrepEditor/Private/DataPrepEditor.h
--- a/Engine/Plugins/Experimental/DataPrepEditor/Source/DataPrepEditor/Private/DataPrepEditor.h
+++ b/Engine/Plugins/Experimental/DataPrepEditor/Source/DataPrepEditor/Private/DataPrepEditor.h
@@ -130,7 +130,6 @@ private:
 	TSharedRef<SDockTab> SpawnTabDataprep(const FSpawnTabArgs& Args);
 	TSharedRef<SDockTab> SpawnTabSceneViewport( const FSpawnTabArgs& Args );
 	TSharedRef<SDockTab> SpawnTabStatistics(const FSpawnTabArgs& Args);
-	TSharedRef<SDockTab> SpawnTabParameterization(const FSpawnTabArgs& Args);
 
 	TSharedRef<FTabManager::FLayout> CreateDataprepLayout();
 	TSharedRef<FTabManager::FLayout> CreateDataprepInstanceLayout();
@@ -255,7 +254,6 @@ private:
 	static const FName DataprepAssetTabId;
 	static const FName SceneViewportTabId;
 	static const FName DataprepStatisticsTabId;
-	static const FName ParameterizationDefaultTabId;
 
 //Temp Code to allow us to work on the nodes while we don't have our own graph.
 public:
diff --git a/Engine/Plugins/Experimental/DataPrepEditor/Source/DataPrepEditor/Private/Widgets/DataprepAssetView.cpp b/Engine/Plugins/Experimental/DataPrepEditor/Source/DataPrepEditor/Private/Widgets/DataprepAssetView.cpp
--- a/Engine/Plugins/Experimental/DataPrepEditor/Source/DataPrepEditor/Private/Widgets/DataprepAssetView.cpp
+++ b/Engine/Plugins/Experimental/DataPrepEditor/Source/DataPrepEditor/Private/Widgets/DataprepAssetView.cpp
@@ -354,6 +354,10 @@ void SDataprepAssetView::Construct( const FArguments& InArgs, UDataprepAssetInte
 	ProducersWidget = SNew(SDataprepProducersWidget, AssetProducers, CommandList)
 		.ColumnSizeData(ColumnSizeData);
 
+	TSharedRef<SDataprepDetailsView> DetailView = SNew(SDataprepDetailsView)
+	.ColumnSizeData( ColumnSizeData )
+	.Object( DataprepAssetInterfacePtr->GetParameterizationObject() );
+
 	TSharedRef<SScrollBar> ScrollBar = SNew(SScrollBar);
 
 	// #ueent_todo: Look at changing the border brushes to add color to this stuff
@@ -373,6 +377,7 @@ void SDataprepAssetView::Construct( const FArguments& InArgs, UDataprepAssetInte
 				+ SScrollBox::Slot()
 				[
 					SNew(SVerticalBox)
+					// Begin - Section for producers
 					+ SVerticalBox::Slot()
 					.Padding( 5.0f )
 					.AutoHeight()
@@ -380,6 +385,9 @@ void SDataprepAssetView::Construct( const FArguments& InArgs, UDataprepAssetInte
 					[
 						ProducersWidget.ToSharedRef()
 					]
+					// End - Section for producers
+
+					// Begin - Section for consumer
 					+ SVerticalBox::Slot()
 					.AutoHeight()
 					[
@@ -405,8 +413,6 @@ void SDataprepAssetView::Construct( const FArguments& InArgs, UDataprepAssetInte
 							.Size( FVector2D( 200, 10 ) )		
 						]
 					]
-
-					// Section for consumer
 					+ SVerticalBox::Slot()
 					.AutoHeight()
 					[
@@ -453,7 +459,66 @@ void SDataprepAssetView::Construct( const FArguments& InArgs, UDataprepAssetInte
 							.Size( FVector2D( 200, 10 ) )		
 						]
 					]
-					// Section for consumer
+					// End - Section for consumer
+
+					// Begin - Section for parameterization
+					+ SVerticalBox::Slot()
+					.AutoHeight()
+					[
+						SNew(SHorizontalBox)
+						+ SHorizontalBox::Slot()
+						.FillWidth(1.0f)
+						.Padding( 0, 10, 0, 0 )
+						.HAlign( EHorizontalAlignment::HAlign_Center )
+						[
+							SNew(SColorBlock)
+							.Color( FLinearColor( 0.9f, 0.9f, 0.9f ) )
+							.Size( FVector2D( 3000, 1 ) )
+						]
+					]
+					+ SVerticalBox::Slot()
+					.AutoHeight()
+					[
+						SNew(SHorizontalBox)
+						+ SHorizontalBox::Slot()
+						[
+							SNew(SSpacer)
+							.Size( FVector2D( 200, 10 ) )		
+						]
+					]
+					+ SVerticalBox::Slot()
+					.AutoHeight()
+					[
+						SNew(SHorizontalBox)
+						+ SHorizontalBox::Slot()
+						.FillWidth(1.0f)
+						.HAlign(EHorizontalAlignment::HAlign_Left)
+						.VAlign(VAlign_Center)
+						.Padding( FMargin( 5.0f, 5.0f, 0, 7.0f ) )
+						[
+							SNew(STextBlock)
+							.Text(LOCTEXT("DataprepAssetView_Consumer_Parameterization", "Parameterization"))
+							.MinDesiredWidth( 200 )
+							.Font( IDetailLayoutBuilder::GetDetailFontBold() )
+						]
+					]
+					+ SVerticalBox::Slot()
+					.AutoHeight()
+					.Padding(10.0f, 5.0f, 0.0f, 5.0f)
+					[
+						SNew(SHorizontalBox)
+						+ SHorizontalBox::Slot()
+						.FillWidth(1.0f)
+						.HAlign(EHorizontalAlignment::HAlign_Left)
+						.VAlign(VAlign_Center)
+						.Padding( FMargin( -5.0f, 5.0f, 20.0f, 5.0f ) )
+						[
+							SNew( SDataprepDetailsView )
+							.Object( DataprepAssetInterfacePtr->GetParameterizationObject() )
+							.ColumnSizeData( ColumnSizeData )
+						]
+					]
+					// End - Section for parameterization
 				]
 			]
 			+SHorizontalBox::Slot()
