Dataprep crash on create proxy mesh operation #jira UE-82279 #rb JeanLuc.Corenthin

diff --git a/Engine/Plugins/Experimental/DataPrepEditor/Source/DataPrepLibraries/Private/DataprepEditingOperations.cpp b/Engine/Plugins/Experimental/DataPrepEditor/Source/DataPrepLibraries/Private/DataprepEditingOperations.cpp
--- a/Engine/Plugins/Experimental/DataPrepEditor/Source/DataPrepLibraries/Private/DataprepEditingOperations.cpp
+++ b/Engine/Plugins/Experimental/DataPrepEditor/Source/DataPrepLibraries/Private/DataprepEditingOperations.cpp
@@ -68,7 +68,7 @@ namespace DatasmithEditingOperationsUtils
 	void FindActorsToMerge(const TArray<AActor*>& ChildrenActors, TArray<AActor*>& ActorsToMerge);
 	void FindActorsToCollapseOrDelete(const TArray<AActor*>& ActorsToVisit, TArray<AActor*>& ActorsToCollapse, TArray<UObject*>& ActorsToDelete );
 	void GetRootActors(UWorld* World, TArray<AActor*>& RootActors);
-	void GetActorsToMerge(UWorld*& World, const TArray<UObject*>& InObjects, TArray<AActor*>& ActorsToMerge, TArray<UPrimitiveComponent*> ComponentsToMerge);
+	void GetActorsToMerge(UWorld*& World, const TArray<UObject*>& InObjects, TArray<AActor*>& ActorsToMerge, TArray<UPrimitiveComponent*>& ComponentsToMerge);
 
 	int32 GetActorDepth(AActor* Actor)
 	{
@@ -297,7 +297,7 @@ void UDataprepCreateProxyMeshOperation::OnExecution_Implementation(const FDatapr
 	DatasmithEditingOperationsUtils::GetActorsToMerge(CurrentWorld, InContext.Objects, ActorsToMerge, ComponentsToMerge);
 
 	// Nothing to do if there is none or only one static mesh actor
-	if(ActorsToMerge.Num() < 2 && ComponentsToMerge.Num() < 2)
+	if(ComponentsToMerge.Num() < 2)
 	{
 		UE_LOG(LogDataprep, Log, TEXT("No static mesh actors to merge"));
 		return;
@@ -671,7 +671,7 @@ namespace DatasmithEditingOperationsUtils
 		}
 	}
 
-	void GetActorsToMerge(UWorld*& World, const TArray<UObject*>& InObjects, TArray<AActor*>& ActorsToMerge, TArray<UPrimitiveComponent*> ComponentsToMerge)
+	void GetActorsToMerge(UWorld*& World, const TArray<UObject*>& InObjects, TArray<AActor*>& ActorsToMerge, TArray<UPrimitiveComponent*>& ComponentsToMerge)
 	{
 		World = nullptr;
 
