Improved performance of taking snapshot and restoring from snapshot
Improved feedback reporting when taking snapshot and restoring from snapshot

diff --git a/Engine/Plugins/Experimental/DataPrepEditor/Source/DataPrepCore/Private/DataPrepAssetProducers.cpp b/Engine/Plugins/Experimental/DataPrepEditor/Source/DataPrepCore/Private/DataPrepAssetProducers.cpp
--- a/Engine/Plugins/Experimental/DataPrepEditor/Source/DataPrepCore/Private/DataPrepAssetProducers.cpp
+++ b/Engine/Plugins/Experimental/DataPrepEditor/Source/DataPrepCore/Private/DataPrepAssetProducers.cpp
@@ -47,50 +47,73 @@ TArray< TWeakObjectPtr< UObject > > UDataprepAssetProducers::Produce( const FDat
 	}
 
 	TArray< TWeakObjectPtr< UObject > > OutAssets;
+	{
+		FDataprepWorkReporter Task( InContext.ProgressReporterPtr, LOCTEXT( "ProducerReport_Importing", "Importing ..." ), (float)AssetProducers.Num(), 1.0f );
 
-	FDataprepWorkReporter Task( InContext.ProgressReporterPtr, LOCTEXT( "ProducerReport_Importing", "Importing ..." ), (float)AssetProducers.Num(), 1.0f );
-
-	TArray< TWeakObjectPtr< UObject > > ProducerAssets;
+		TArray< TWeakObjectPtr< UObject > > ProducerAssets;
 
-	for ( FDataprepAssetProducer& AssetProducer : AssetProducers )
-	{
-		if( UDataprepContentProducer* Producer = AssetProducer.Producer )
+		for ( FDataprepAssetProducer& AssetProducer : AssetProducers )
 		{
-			Task.ReportNextStep( FText::Format( LOCTEXT( "ProducerReport_ImportingX", "Importing {0} ..."), FText::FromString( Producer->GetName() ) ) );
+			if( UDataprepContentProducer* Producer = AssetProducer.Producer )
+			{
+				Task.ReportNextStep( FText::Format( LOCTEXT( "ProducerReport_ImportingX", "Importing {0} ..."), FText::FromString( Producer->GetName() ) ) );
 
-			// Run producer if enabled and, if superseded, superseder is disabled
-			const bool bIsOkToRun = AssetProducer.bIsEnabled &&	( AssetProducer.SupersededBy == INDEX_NONE || !AssetProducers[AssetProducer.SupersededBy].bIsEnabled );
+				// Run producer if enabled and, if superseded, superseder is disabled
+				const bool bIsOkToRun = AssetProducer.bIsEnabled &&	( AssetProducer.SupersededBy == INDEX_NONE || !AssetProducers[AssetProducer.SupersededBy].bIsEnabled );
 
-			if ( bIsOkToRun )
-			{
-				if( !Producer->Produce( InContext, OutAssets ) )
+				if ( bIsOkToRun )
 				{
-					if ( Task.IsWorkCancelled() )
+					if( !Producer->Produce( InContext, OutAssets ) )
 					{
-						FText Report = FText::Format(LOCTEXT( "ProducerReport_Cancelled", "Producer {0} was cancelled during execution."), FText::FromString( Producer->GetName() ) );
-						InContext.LoggerPtr->LogInfo( Report, *Producer );
-						break;
-					}
-					else
-					{
-						FText Report = FText::Format(LOCTEXT( "ProducerReport_Failed", "Producer {0} failed to execute."), FText::FromString( Producer->GetName() ) );
-						InContext.LoggerPtr->LogError( Report, *Producer );
+						if ( Task.IsWorkCancelled() )
+						{
+							FText Report = FText::Format(LOCTEXT( "ProducerReport_Cancelled", "Producer {0} was cancelled during execution."), FText::FromString( Producer->GetName() ) );
+							InContext.LoggerPtr->LogInfo( Report, *Producer );
+							break;
+						}
+						else
+						{
+							FText Report = FText::Format(LOCTEXT( "ProducerReport_Failed", "Producer {0} failed to execute."), FText::FromString( Producer->GetName() ) );
+							InContext.LoggerPtr->LogError( Report, *Producer );
+						}
 					}
 				}
 			}
+			else
+			{
+				Task.ReportNextStep( LOCTEXT( "ProducerReport_Skipped", "Skipped invalid producer ...") );
+			}
 		}
-		else
+
+		// Skip finalization of import if user has canceled process
+		if(IDataprepProgressReporter* ProgressReporter = InContext.ProgressReporterPtr.Get())
 		{
-			Task.ReportNextStep( LOCTEXT( "ProducerReport_Skipped", "Skipped invalid producer ...") );
+			if(ProgressReporter->IsWorkCancelled())
+			{
+				TArray<UObject*> ObjectsToDelete;
+				for ( TWeakObjectPtr<UObject> WeakObjectPtr : OutAssets )
+				{
+					if ( UObject* Object = WeakObjectPtr.Get() )
+					{
+						ObjectsToDelete.Add( Object );
+					}
+				}
+
+				FDataprepCoreUtils::PurgeObjects( MoveTemp( ObjectsToDelete ) );
+
+				OutAssets.Empty();
+
+				return OutAssets;
+			}
 		}
-	}
 
-	for ( TWeakObjectPtr<UObject> WeakObjectPtr : OutAssets )
-	{
-		if ( UObject* Object = WeakObjectPtr.Get() )
+		for ( TWeakObjectPtr<UObject> WeakObjectPtr : OutAssets )
 		{
-			// We must ensure that public face of an asset is public
-			Object->SetFlags( RF_Public );
+			if ( UObject* Object = WeakObjectPtr.Get() )
+			{
+				// We must ensure that public face of an asset is public
+				Object->SetFlags( RF_Public );
+			}
 		}
 	}
 
diff --git a/Engine/Plugins/Experimental/DataPrepEditor/Source/DataPrepCore/Private/DataprepCoreUtils.cpp b/Engine/Plugins/Experimental/DataPrepEditor/Source/DataPrepCore/Private/DataprepCoreUtils.cpp
--- a/Engine/Plugins/Experimental/DataPrepEditor/Source/DataPrepCore/Private/DataprepCoreUtils.cpp
+++ b/Engine/Plugins/Experimental/DataPrepEditor/Source/DataPrepCore/Private/DataprepCoreUtils.cpp
@@ -158,13 +158,13 @@ bool FDataprepCoreUtils::IsAsset(UObject* Object)
 	return true;
 }
 
-FDataprepWorkReporter::FDataprepWorkReporter(const TSharedPtr<IDataprepProgressReporter>& InReporter, const FText& InDescription, float InAmountOfWork, float InIncrementOfWork)
+FDataprepWorkReporter::FDataprepWorkReporter(const TSharedPtr<IDataprepProgressReporter>& InReporter, const FText& InDescription, float InAmountOfWork, float InIncrementOfWork, bool bInterruptible )
 	: Reporter( InReporter )
 	, DefaultIncrementOfWork( InIncrementOfWork )
 {
 	if( Reporter.IsValid())
 	{
-		Reporter->BeginWork( InDescription, InAmountOfWork );
+		Reporter->BeginWork( InDescription, InAmountOfWork, bInterruptible );
 	}
 }
 
@@ -240,10 +240,10 @@ void FDataprepCoreUtils::FDataprepLogger::LogError(const FText& InLogText,  cons
 	UE_LOG( LogDataprepCore, Error, TEXT("%s : %s"), *InObject.GetName(), *InLogText.ToString() );
 }
 
-void FDataprepCoreUtils::FDataprepProgressUIReporter::BeginWork( const FText& InTitle, float InAmountOfWork )
+void FDataprepCoreUtils::FDataprepProgressUIReporter::BeginWork( const FText& InTitle, float InAmountOfWork, bool bInterruptible )
 {
 	ProgressTasks.Emplace( new FScopedSlowTask( InAmountOfWork, InTitle, true, FeedbackContext.IsValid() ? *FeedbackContext.Get() : *GWarn ) );
-	ProgressTasks.Last()->MakeDialog(true);
+	ProgressTasks.Last()->MakeDialog( bInterruptible );
 }
 
 void FDataprepCoreUtils::FDataprepProgressUIReporter::EndWork()
@@ -278,7 +278,7 @@ FFeedbackContext* FDataprepCoreUtils::FDataprepProgressUIReporter::GetFeedbackCo
 	return FeedbackContext.IsValid() ? FeedbackContext.Get() : GWarn;
 }
 
-void FDataprepCoreUtils::FDataprepProgressTextReporter::BeginWork( const FText& InTitle, float InAmountOfWork )
+void FDataprepCoreUtils::FDataprepProgressTextReporter::BeginWork( const FText& InTitle, float InAmountOfWork, bool /*bInterruptible*/ )
 {
 	UE_LOG( LogDataprepCore, Log, TEXT("Start: %s ..."), *InTitle.ToString() );
 	++TaskDepth;
@@ -335,6 +335,9 @@ void FDataprepCoreUtils::BuildAssets(const TArray<TWeakObjectPtr<UObject>>& Asse
 		}
 	}
 
+	int32 AssetToBuildCount = MaterialInterfaces.Num() + StaticMeshes.Num();
+	FDataprepWorkReporter Task( ProgressReporterPtr, LOCTEXT( "BuildAssets_Building", "Building assets ..." ), (float)AssetToBuildCount, 1.0f, false );
+
 	// Force compilation of materials which have no render proxy
 	if(MaterialInterfaces.Num() > 0)
 	{
@@ -349,8 +352,14 @@ void FDataprepCoreUtils::BuildAssets(const TArray<TWeakObjectPtr<UObject>>& Asse
 			return false;
 		};
 
+		int32 AssetBuiltCount = 0;
+		AssetToBuildCount = MaterialInterfaces.Num();
+
 		for(UMaterialInterface* MaterialInterface : MaterialInterfaces)
 		{
+			++AssetBuiltCount;
+			Task.ReportNextStep(FText::Format(LOCTEXT( "BuildAssets_Building_Materials", "Building materials ({0} / {1})" ), AssetBuiltCount, AssetToBuildCount), 1.0f);
+
 			if( MustCompile( MaterialInterface ) )
 			{
 				FPropertyChangedEvent EmptyPropertyUpdateStruct( nullptr );
@@ -360,7 +369,13 @@ void FDataprepCoreUtils::BuildAssets(const TArray<TWeakObjectPtr<UObject>>& Asse
 	}
 
 	// Build static meshes
-	DataprepCorePrivateUtils::BuildStaticMeshes( StaticMeshes, TFunction<bool(UStaticMesh*)>() );
+	int32 AssetBuiltCount = 0;
+	AssetToBuildCount = StaticMeshes.Num();
+	DataprepCorePrivateUtils::BuildStaticMeshes(StaticMeshes, [&](UStaticMesh* StaticMesh) -> bool {
+		++AssetBuiltCount;
+		Task.ReportNextStep(FText::Format(LOCTEXT( "BuildAssets_Building_Meshes", "Building static meshes ({0} / {1})" ), AssetBuiltCount, AssetToBuildCount), 1.0f);
+		return true;
+	});
 }
 
 #undef LOCTEXT_NAMESPACE
\ No newline at end of file
diff --git a/Engine/Plugins/Experimental/DataPrepEditor/Source/DataPrepCore/Public/DataprepCoreUtils.h b/Engine/Plugins/Experimental/DataPrepEditor/Source/DataPrepCore/Public/DataprepCoreUtils.h
--- a/Engine/Plugins/Experimental/DataPrepEditor/Source/DataPrepCore/Public/DataprepCoreUtils.h
+++ b/Engine/Plugins/Experimental/DataPrepEditor/Source/DataPrepCore/Public/DataprepCoreUtils.h
@@ -130,7 +130,7 @@ public:
 		}
 
 		// Begin IDataprepProgressReporter interface
-		virtual void BeginWork( const FText& InTitle, float InAmountOfWork ) override;
+		virtual void BeginWork( const FText& InTitle, float InAmountOfWork, bool bInterruptible = true ) override;
 		virtual void EndWork() override;
 		virtual void ReportProgress( float Progress, const FText& InMessage ) override;
 		virtual bool IsWorkCancelled() override;
@@ -157,7 +157,7 @@ public:
 		}
 
 		// Begin IDataprepProgressReporter interface
-		virtual void BeginWork( const FText& InTitle, float InAmountOfWork ) override;
+		virtual void BeginWork( const FText& InTitle, float InAmountOfWork, bool bInterruptible = true ) override;
 		virtual void EndWork() override;
 		virtual void ReportProgress( float Progress, const FText& InMessage ) override;
 		virtual bool IsWorkCancelled() override;
diff --git a/Engine/Plugins/Experimental/DataPrepEditor/Source/DataPrepCore/Public/IDataprepProgressReporter.h b/Engine/Plugins/Experimental/DataPrepEditor/Source/DataPrepCore/Public/IDataprepProgressReporter.h
--- a/Engine/Plugins/Experimental/DataPrepEditor/Source/DataPrepCore/Public/IDataprepProgressReporter.h
+++ b/Engine/Plugins/Experimental/DataPrepEditor/Source/DataPrepCore/Public/IDataprepProgressReporter.h
@@ -22,7 +22,7 @@ public:
 	 * @param InDescription		Text describing the work about to begin
 	 * @param InAmountOfWork	Expected total amount of work
 	 */
-	virtual void BeginWork( const FText& InDescription, float InAmountOfWork ) = 0;
+	virtual void BeginWork( const FText& InDescription, float InAmountOfWork, bool bInterruptible = true ) = 0;
 
 	/** Indicates the end of the work */
 	virtual void EndWork() = 0;
@@ -56,7 +56,7 @@ public:
 	 * @param InAmountOfWork		Total amount of work for the task
 	 * @param InIncrementOfWork		Amount of incremental work at each step within the task
 	 */
-	FDataprepWorkReporter( const TSharedPtr<IDataprepProgressReporter>& InReporter, const FText& InDescription, float InAmountOfWork, float InIncrementOfWork = 1.0f );
+	FDataprepWorkReporter( const TSharedPtr<IDataprepProgressReporter>& InReporter, const FText& InDescription, float InAmountOfWork, float InIncrementOfWork = 1.0f, bool bInterruptible = true );
 
 	~FDataprepWorkReporter();
 
diff --git a/Engine/Plugins/Experimental/DataPrepEditor/Source/DataPrepEditor/Private/DataprepSnapshot.cpp b/Engine/Plugins/Experimental/DataPrepEditor/Source/DataPrepEditor/Private/DataprepSnapshot.cpp
--- a/Engine/Plugins/Experimental/DataPrepEditor/Source/DataPrepEditor/Private/DataprepSnapshot.cpp
+++ b/Engine/Plugins/Experimental/DataPrepEditor/Source/DataPrepEditor/Private/DataprepSnapshot.cpp
@@ -4,9 +4,12 @@
 
 #include "DataprepCoreUtils.h"
 #include "DataprepEditorLogCategory.h"
+#include "DataprepEditorUtils.h"
 
 #include "ActorEditorUtils.h"
 #include "AutoReimport/AutoReimportManager.h"
+#include "Async/Async.h"
+#include "Async/Future.h"
 #include "Async/ParallelFor.h"
 #include "Editor/UnrealEdEngine.h"
 #include "EngineUtils.h"
@@ -16,6 +19,8 @@
 #include "HAL/FileManager.h"
 #include "GenericPlatform/GenericPlatformOutputDevices.h"
 #include "GenericPlatform/GenericPlatformTime.h"
+#include "Materials/MaterialFunction.h"
+#include "Materials/MaterialFunctionInstance.h"
 #include "Materials/MaterialInstance.h"
 #include "MaterialShared.h"
 #include "Misc/Compression.h"
@@ -35,6 +40,8 @@
 enum class EDataprepAssetClass : uint8 {
 	EDataprep,
 	ETexture,
+	EMaterialFunction,
+	EMaterialFunctionInstance,
 	EMaterial,
 	EMaterialInstance,
 	EStaticMesh,
@@ -469,6 +476,14 @@ void FDataprepEditor::TakeSnapshot()
 		{
 			return EDataprepAssetClass::EStaticMesh;
 		}
+		else if (AssetClass->IsChildOf(UMaterialFunction::StaticClass()))
+		{
+			return EDataprepAssetClass::EMaterialFunction;
+		}
+		else if (AssetClass->IsChildOf(UMaterialFunctionInstance::StaticClass()))
+		{
+			return EDataprepAssetClass::EMaterialFunctionInstance;
+		}
 		else if (AssetClass->IsChildOf(UMaterial::StaticClass()))
 		{
 			return EDataprepAssetClass::EMaterial;
@@ -502,68 +517,109 @@ void FDataprepEditor::TakeSnapshot()
 		}
 	}
 
-	FCriticalSection GlobalLock;
 	TAtomic<bool> bGlobalIsValid(true);
-	ParallelFor(Assets.Num(),
-		[&](int32 AssetIndex)
 	{
-		if (!bGlobalIsValid)
-		{
-			return;
-		}
+		FText Message = LOCTEXT("SaveSnapshot_SaveAssets", "Snapshot : Caching assets ...");
+		SlowTask.EnterProgressFrame( 40.0f, Message );
+
+		FScopedSlowTask SlowSaveAssetTask( (float)Assets.Num(), Message );
+		SlowSaveAssetTask.MakeDialog(false);
+
+		TArray<TFuture<bool>> AsyncTasks;
+		AsyncTasks.Reserve( Assets.Num() );
 
-		if(UObject* AssetObject = Assets[AssetIndex].Get())
+		for (TWeakObjectPtr<UObject> AssetObjectPtr : Assets)
 		{
-			EObjectFlags ObjectFlags = AssetObject->GetFlags();
-			AssetObject->ClearFlags(RF_Transient);
-			AssetObject->SetFlags(RF_Public);
+			AsyncTasks.Emplace(
+				Async(
+					EAsyncExecution::LargeThreadPool,
+					[this, AssetObjectPtr, &bGlobalIsValid]()
+					{
+						if(UObject* AssetObject = AssetObjectPtr.Get())
+						{
+							if ( bGlobalIsValid.Load(EMemoryOrder::Relaxed) )
+							{
+								EObjectFlags ObjectFlags = AssetObject->GetFlags();
+								AssetObject->ClearFlags(RF_Transient);
+								AssetObject->SetFlags(RF_Public);
 
-			FSoftObjectPath AssetPath( AssetObject );
-			const FString AssetPathString = AssetPath.GetAssetPathString();
-			UE_LOG( LogDataprepEditor, Verbose, TEXT("Saving asset %s"), *AssetPathString );
+								FSoftObjectPath AssetPath( AssetObject );
+								const FString AssetPathString = AssetPath.GetAssetPathString();
+								UE_LOG( LogDataprepEditor, Verbose, TEXT("Saving asset %s"), *AssetPathString );
 
-			bool bLocalIsValid = false;
+								bool bLocalIsValid = false;
 
-			// Serialize asset
-			{
-				TArray<uint8> SerializedData;
-				DataprepSnapshotUtil::WriteSnapshotData( AssetObject, SerializedData );
+								// Serialize asset
+								{
+									TArray<uint8> SerializedData;
+									DataprepSnapshotUtil::WriteSnapshotData( AssetObject, SerializedData );
 
-				FString AssetFilePath = DataprepSnapshotUtil::BuildAssetFileName( TempDir, AssetPathString );
+									FString AssetFilePath = DataprepSnapshotUtil::BuildAssetFileName( this->TempDir, AssetPathString );
 
-				bLocalIsValid = FFileHelper::SaveArrayToFile( SerializedData, *AssetFilePath );
-			}
+									bLocalIsValid = FFileHelper::SaveArrayToFile( SerializedData, *AssetFilePath );
+								}
 
-			AssetObject->ClearFlags( RF_AllFlags );
-			AssetObject->SetFlags( ObjectFlags );
+								AssetObject->ClearFlags( RF_AllFlags );
+								AssetObject->SetFlags( ObjectFlags );
 
-			if(!bLocalIsValid)
+								return bLocalIsValid;
+							}
+							else
+							{
+								return false;
+							}
+						}
+
+						return true;
+					}
+				)
+			);
+		}
+
+		for (int32 Index = 0; Index < AsyncTasks.Num(); ++Index)
+		{
+			if(UObject* AssetObject = Assets[Index].Get())
 			{
-				UE_LOG( LogDataprepEditor, Log, TEXT("Failed to save %s"), *AssetPathString );
-				bGlobalIsValid = false;
-				return;
-			}
+				SlowSaveAssetTask.EnterProgressFrame();
+
+				const FSoftObjectPath AssetPath( AssetObject );
+				const FString AssetPathString = AssetPath.GetAssetPathString();
+
+				// Wait the result of the async task
+				if(!AsyncTasks[Index].Get())
+				{
+					UE_LOG( LogDataprepEditor, Log, TEXT("Failed to save %s"), *AssetPathString );
 
-			UE_LOG( LogDataprepEditor, Verbose, TEXT("Asset %s successfully saved"), *AssetPathString );
+					bGlobalIsValid = false;
+					break;
+				}
+				else
+				{
+					UE_LOG( LogDataprepEditor, Verbose, TEXT("Asset %s successfully saved"), *AssetPathString );
+				}
+			}
 		}
 	}
-	);
 
 	ContentSnapshot.bIsValid = bGlobalIsValid;
 
 	// Serialize world if applicable
 	if(ContentSnapshot.bIsValid)
 	{
-		SlowTask.EnterProgressFrame( 50.0f, LOCTEXT("SaveSnapshot_World", "Snapshot : caching level ...") );
+		FText Message = LOCTEXT("SaveSnapshot_World", "Snapshot : caching level ...");
+		SlowTask.EnterProgressFrame( 50.0f, Message );
 		UE_LOG( LogDataprepEditor, Verbose, TEXT("Saving preview world") );
 
+		FScopedSlowTask SlowSaveAssetTask( (float)PreviewWorld->GetCurrentLevel()->Actors.Num(), Message );
+		SlowSaveAssetTask.MakeDialog(false);
+
 		PreviewWorld->ClearFlags(RF_Transient);
 		{
 			// Code inspired from UUnrealEdEngine::edactCopySelected
 			FStringOutputDevice Ar;
 			uint32 ExportFlags = PPF_DeepCompareInstances | PPF_ExportsNotFullyQualified | PPF_IncludeTransient;
 			const FDataprepExportObjectInnerContext Context( PreviewWorld.Get() );
-			UExporter::ExportToOutputDevice( &Context, PreviewWorld.Get(), NULL, Ar, TEXT("copy"), 0, ExportFlags);
+			UExporter::ExportToOutputDevice( &Context, PreviewWorld.Get(), nullptr, Ar, TEXT("copy"), 0, ExportFlags);
 
 			// Save text into file
 			FString PackageFilePath = DataprepSnapshotUtil::BuildAssetFileName( TempDir, GetTransientContentFolder() / SessionID ) + TEXT(".asc");
@@ -731,24 +787,22 @@ void FDataprepEditor::RestoreFromSnapshot(bool bUpdateViewport)
 		check( FFileHelper::LoadFileToString(FileBuffer, *PackageFilePath) );
 
 		// Set the GWorld to the preview world since ULevelFactory::FactoryCreateText uses GWorld
-		UWorld* CachedWorld = GWorld;
+		UWorld* PrevGWorld = GWorld;
 		GWorld = PreviewWorld.Get();
 
 		// Cache and disable recording of transaction
-		UTransactor* NormalTransactor = GEditor->Trans;
-		GEditor->Trans = nullptr;
+		TGuardValue<UTransactor*> NormalTransactor( GEditor->Trans, nullptr );
 
 		// Cache and disable warnings from LogExec because ULevelFactory::FactoryCreateText is pretty verbose on harmless warnings
 		ELogVerbosity::Type PrevLogExecVerbosity = LogExec.GetVerbosity();
 		LogExec.SetVerbosity( ELogVerbosity::Error );
 
 		// Cache and disable Editor's selection
-		bool PrevEdSelectionLock = GEdSelectionLock;
-		GEdSelectionLock = false;
+		TGuardValue<bool> EdSelectionLock( GEdSelectionLock, true );
 
 		const TCHAR* Paste = *FileBuffer;
 		ULevelFactory* Factory = NewObject<ULevelFactory>();
-		Factory->FactoryCreateText( ULevel::StaticClass(), WorldLevel, WorldLevel->GetFName(), RF_Transactional, NULL, TEXT("paste"), Paste, Paste + FileBuffer.Len(), GWarn/*FGenericPlatformOutputDevices::GetFeedbackContext()*/);
+		Factory->FactoryCreateText( ULevel::StaticClass(), WorldLevel, WorldLevel->GetFName(), RF_Transactional, NULL, TEXT("paste"), Paste, Paste + FileBuffer.Len(), GWarn );
 
 		// Restore LogExec verbosity
 		LogExec.SetVerbosity( PrevLogExecVerbosity );
@@ -756,14 +810,8 @@ void FDataprepEditor::RestoreFromSnapshot(bool bUpdateViewport)
 		// Reinstate old BSP update setting, and force a rebuild - any levels whose geometry has changed while pasting will be rebuilt
 		GetMutableDefault<ULevelEditorMiscSettings>()->bBSPAutoUpdate = bBSPAutoUpdate;
 
-		// Restore transaction's recording
-		GEditor->Trans = NormalTransactor;
-
-		// Restore Editor's selection
-		GEdSelectionLock = PrevEdSelectionLock;
-
-		// Reset the GWorld to its previous value
-		GWorld = CachedWorld;
+		// Restore GWorld
+		GWorld = PrevGWorld;
 	}
 	UE_LOG( LogDataprepEditor, Verbose, TEXT("Level loaded") );
 
