Integrating CL# 9758495 from Dev-Enterprise
Fix dataprep background color being the same as disabled properties color

diff --git a/Engine/Plugins/Experimental/DataPrepEditor/Source/DataPrepEditor/Private/DataprepEditorStyle.cpp b/Engine/Plugins/Experimental/DataPrepEditor/Source/DataPrepEditor/Private/DataprepEditorStyle.cpp
--- a/Engine/Plugins/Experimental/DataPrepEditor/Source/DataPrepEditor/Private/DataprepEditorStyle.cpp
+++ b/Engine/Plugins/Experimental/DataPrepEditor/Source/DataPrepEditor/Private/DataprepEditorStyle.cpp
@@ -61,8 +61,8 @@ void FDataprepEditorStyle::Initialize()
 	{
 		StyleSet->Set("DataprepAction.Padding", 2.f);
 		StyleSet->Set("DataprepAction.OutlineColor", FLinearColor(0.f, 0.f, 0.f));
-		StyleSet->Set("DataprepActionBlock.TitleBackgroundColor", FLinearColor(0.095307f, 0.095307f, 0.095307f));
-		StyleSet->Set("DataprepActionBlock.ContentBackgroundColor", FLinearColor(0.215861f, 0.215861f, 0.215861f));
+		StyleSet->Set("DataprepActionBlock.TitleBackgroundColor", FLinearColor(0.065307f, 0.065307f, 0.065307f));
+		StyleSet->Set("DataprepActionBlock.ContentBackgroundColor", FLinearColor(0.115861f, 0.115861f, 0.115861f));
 		FTextBlockStyle TilteTextBlockStyle = FEditorStyle::GetWidgetStyle< FTextBlockStyle >("NormalText");
 		TilteTextBlockStyle.SetFont(FCoreStyle::GetDefaultFontStyle("Bold", 14));
 		StyleSet->Set("DataprepActionBlock.TitleTextBlockStyle", TilteTextBlockStyle);
