Fixing build warning

diff --git a/Engine/Plugins/Experimental/DataPrepEditor/Source/DataPrepLibraries/Private/DataPrepOperationsLibrary.cpp b/Engine/Plugins/Experimental/DataPrepEditor/Source/DataPrepLibraries/Private/DataPrepOperationsLibrary.cpp
--- a/Engine/Plugins/Experimental/DataPrepEditor/Source/DataPrepLibraries/Private/DataPrepOperationsLibrary.cpp
+++ b/Engine/Plugins/Experimental/DataPrepEditor/Source/DataPrepLibraries/Private/DataPrepOperationsLibrary.cpp
@@ -71,9 +71,9 @@ namespace DataprepOperationsLibraryUtil
 			{
 				SelectedMeshes.Add( StaticMesh );
 			}
-			else if ( UStaticMeshComponent* StaticMeshComponent = Cast<UStaticMeshComponent>(Object) )
+			else if ( Object->IsA(UStaticMeshComponent::StaticClass()) )
 			{
-				if (UStaticMesh* StaticMesh = StaticMeshComponent->GetStaticMesh())
+				if ((StaticMesh = Cast<UStaticMeshComponent>(Object)->GetStaticMesh()) != nullptr )
 				{
 					SelectedMeshes.Add(StaticMesh);
 				}
