Dataprep parameterization implemented a proper validation of value type for the bindings.

diff --git a/Engine/Plugins/Experimental/DataPrepEditor/Source/DataPrepCore/Private/DataPrepAsset.cpp b/Engine/Plugins/Experimental/DataPrepEditor/Source/DataPrepCore/Private/DataPrepAsset.cpp
--- a/Engine/Plugins/Experimental/DataPrepEditor/Source/DataPrepCore/Private/DataPrepAsset.cpp
+++ b/Engine/Plugins/Experimental/DataPrepEditor/Source/DataPrepCore/Private/DataPrepAsset.cpp
@@ -231,9 +231,9 @@ void UDataprepAsset::RemoveObjectPropertyFromParameterization(UDataprepParameter
 	Parameterization->RemoveBindedObjectProperty( Object, InPropertyChain );
 }
 
-void UDataprepAsset::GetExistingParameterNamesForType(UClass* PropertyClass, TSet<FString>& OutValidExistingNames, TSet<FString>& OutInvalidNames) const
+void UDataprepAsset::GetExistingParameterNamesForType(UProperty* Property, TSet<FString>& OutValidExistingNames, TSet<FString>& OutInvalidNames) const
 {
-	Parameterization->GetExistingParameterNamesForType( PropertyClass, OutValidExistingNames, OutInvalidNames );
+	Parameterization->GetExistingParameterNamesForType( Property, OutValidExistingNames, OutInvalidNames );
 }
 
 void UDataprepAsset::OnDataprepBlueprintChanged( UBlueprint* InBlueprint )
diff --git a/Engine/Plugins/Experimental/DataPrepEditor/Source/DataPrepCore/Private/Parameterization/DataprepParameterization.cpp b/Engine/Plugins/Experimental/DataPrepEditor/Source/DataPrepCore/Private/Parameterization/DataprepParameterization.cpp
--- a/Engine/Plugins/Experimental/DataPrepEditor/Source/DataPrepCore/Private/Parameterization/DataprepParameterization.cpp
+++ b/Engine/Plugins/Experimental/DataPrepEditor/Source/DataPrepCore/Private/Parameterization/DataprepParameterization.cpp
@@ -7,6 +7,7 @@
 #include "DataprepParameterizationArchive.h"
 
 #include "CoreGlobals.h"
+#include "DataprepCoreLogCategory.h"
 #include "Editor.h"
 #include "Engine/Engine.h"
 #include "Math/UnrealMathUtility.h"
@@ -21,6 +22,151 @@
 
 namespace DataprepParameterization
 {
+	void PopulateValueTypeValidationData(UArrayProperty* CurrentProperty, TArray<UObject*>& ValueTypeValidationData);
+	void PopulateValueTypeValidationData(USetProperty* CurrentProperty, TArray<UObject*>& ValueTypeValidationData);
+	void PopulateValueTypeValidationData(UMapProperty* CurrentProperty, TArray<UObject*>& ValueTypeValidationData);
+	void PopulateValueTypeValidationData(USoftClassProperty* CurrentProperty, TArray<UObject*>& ValueTypeValidationData);
+	void PopulateValueTypeValidationData(UClassProperty* CurrentProperty, TArray<UObject*>& ValueTypeValidationData);
+	void PopulateValueTypeValidationData(UObjectPropertyBase* CurrentProperty, TArray<UObject*>& ValueTypeValidationData);
+	void PopulateValueTypeValidationData(UEnumProperty* CurrentProperty, TArray<UObject*>& ValueTypeValidationData);
+	void PopulateValueTypeValidationData(UStructProperty* CurrentProperty, TArray<UObject*>& ValueTypeValidationData);
+	void PopulateValueTypeValidationData(UProperty* CurrentProperty, TArray<UObject*>& ValueTypeValidationData);
+
+	void PopulateValueTypeValidationData(UArrayProperty* CurrentProperty, TArray<UObject*>& ValueTypeValidationData)
+	{
+		if ( CurrentProperty )
+		{
+			ValueTypeValidationData.Reserve( 2 );
+			ValueTypeValidationData.Add( CurrentProperty->GetClass() );
+			PopulateValueTypeValidationData( CurrentProperty->Inner, ValueTypeValidationData );
+		}
+	}
+
+	void PopulateValueTypeValidationData(USetProperty* CurrentProperty, TArray<UObject*>& ValueTypeValidationData)
+	{
+		if ( CurrentProperty )
+		{
+			ValueTypeValidationData.Reserve( 2 );
+			ValueTypeValidationData.Add( CurrentProperty->GetClass() );
+			PopulateValueTypeValidationData( CurrentProperty->ElementProp, ValueTypeValidationData );
+		}
+	}
+
+	void PopulateValueTypeValidationData(UMapProperty* CurrentProperty, TArray<UObject*>& ValueTypeValidationData)
+	{
+		if ( CurrentProperty )
+		{
+			ValueTypeValidationData.Reserve( 3 );
+			ValueTypeValidationData.Add( CurrentProperty->GetClass() );
+			PopulateValueTypeValidationData( CurrentProperty->KeyProp, ValueTypeValidationData );
+			PopulateValueTypeValidationData( CurrentProperty->ValueProp, ValueTypeValidationData );
+		}
+	}
+
+	void PopulateValueTypeValidationData(USoftClassProperty* CurrentProperty, TArray<UObject*>& ValueTypeValidationData)
+	{
+		if ( CurrentProperty )
+		{
+			ValueTypeValidationData.Reserve( 2 );
+			ValueTypeValidationData.Add( CurrentProperty->GetClass() );
+			// Property class don't matter here since it's always UObject::StaticClass
+			ValueTypeValidationData.Add( CurrentProperty->MetaClass );
+		}
+	}
+
+	void PopulateValueTypeValidationData(UClassProperty* CurrentProperty, TArray<UObject*>& ValueTypeValidationData)
+	{
+		if ( CurrentProperty )
+		{
+			ValueTypeValidationData.Reserve( 2 );
+			ValueTypeValidationData.Add( CurrentProperty->GetClass() );
+			ValueTypeValidationData.Add( CurrentProperty->PropertyClass );
+			ValueTypeValidationData.Add( CurrentProperty->MetaClass );
+		}
+	}
+
+	void PopulateValueTypeValidationData(UObjectPropertyBase* CurrentProperty, TArray<UObject*>& ValueTypeValidationData)
+	{
+		if ( CurrentProperty )
+		{
+			ValueTypeValidationData.Reserve( 2 );
+			ValueTypeValidationData.Add( CurrentProperty->GetClass() );
+			ValueTypeValidationData.Add( CurrentProperty->PropertyClass );
+		}
+	}
+
+	void PopulateValueTypeValidationData(UEnumProperty* CurrentProperty, TArray<UObject*>& ValueTypeValidationData)
+	{
+		if ( CurrentProperty )
+		{
+			ValueTypeValidationData.Reserve( 3 );
+			ValueTypeValidationData.Add( CurrentProperty->GetClass() );
+			ValueTypeValidationData.Add( CurrentProperty->GetEnum() );
+			ValueTypeValidationData.Add( CurrentProperty->GetUnderlyingProperty() );
+		}
+	}
+
+	void PopulateValueTypeValidationData(UStructProperty* CurrentProperty, TArray<UObject*>& ValueTypeValidationData)
+	{
+		if ( CurrentProperty )
+		{
+			ValueTypeValidationData.Reserve( 2 );
+			ValueTypeValidationData.Add( CurrentProperty->GetClass() );
+			ValueTypeValidationData.Add( CurrentProperty->Struct );
+		}
+	}
+
+	void PopulateValueTypeValidationData(UProperty* CurrentProperty, TArray<UObject*>& ValueTypeValidationData)
+	{
+		if ( CurrentProperty )
+		{
+			int32 NumberOfObject = ValueTypeValidationData.Num();
+			UClass* CurrentClass = CurrentProperty->GetClass();
+			while ( CurrentClass && NumberOfObject == ValueTypeValidationData.Num() )
+			{
+				if ( CurrentClass == UArrayProperty::StaticClass() )
+				{
+					PopulateValueTypeValidationData( static_cast<UArrayProperty*>( CurrentProperty ), ValueTypeValidationData );
+				}
+				else if ( CurrentClass == USetProperty::StaticClass() )
+				{
+					PopulateValueTypeValidationData( static_cast<USetProperty*>( CurrentProperty ), ValueTypeValidationData );
+				}
+				else if ( CurrentClass == UMapProperty::StaticClass() )
+				{
+					PopulateValueTypeValidationData( static_cast<UMapProperty*>( CurrentProperty ), ValueTypeValidationData );
+				}
+				else if ( CurrentClass == USoftClassProperty::StaticClass() )
+				{
+					PopulateValueTypeValidationData( static_cast<USoftClassProperty*>( CurrentProperty ), ValueTypeValidationData );
+				}
+				else if ( CurrentClass == UClassProperty::StaticClass() )
+				{
+					PopulateValueTypeValidationData( static_cast<UClassProperty*>( CurrentProperty ), ValueTypeValidationData );
+				}
+				else if ( CurrentClass == UObjectPropertyBase::StaticClass() )
+				{
+					PopulateValueTypeValidationData( static_cast<UObjectPropertyBase*>( CurrentProperty ), ValueTypeValidationData );
+				}
+				else if ( CurrentClass == UEnumProperty::StaticClass() )
+				{
+					PopulateValueTypeValidationData( static_cast<UEnumProperty*>( CurrentProperty ), ValueTypeValidationData );
+				}
+				else if ( CurrentClass == UStructProperty::StaticClass() )
+				{
+					PopulateValueTypeValidationData( static_cast<UStructProperty*>( CurrentProperty ), ValueTypeValidationData );
+				}
+				else
+				{
+					ValueTypeValidationData.Add( CurrentProperty->GetClass() );
+				}
+
+				CurrentClass = CurrentClass->GetSuperClass();
+			}
+		}
+	}
+
+
 	void* GetAddressOf(const UArrayProperty& Property, void* BaseAddress, int32 ContainerIndex)
 	{
 		void* AddressOfArray = Property.ContainerPtrToValuePtr<void*>( BaseAddress, 0 );
@@ -230,7 +376,7 @@ namespace DataprepParameterization
 	 */
 	UProperty* GetPropertyFromBinding(FDataprepParameterizationBinding& Binding, void*& OutPropertyValueAddress)
 	{
-		if ( !Binding.ObjectBinded || !Binding. ValueType )
+		if ( !Binding.ObjectBinded || Binding.ValueTypeValidationData.Num() == 0 )
 		{
 			return nullptr;
 		}
@@ -279,10 +425,18 @@ namespace DataprepParameterization
 		if ( LevelIndex == PropertyChain.Num() && CurrentOuter )
 		{
 			PropetyAtCurrentLevel = PropertyChain.Last().CachedProperty.Get();
-			if ( PropetyAtCurrentLevel->GetClass() == Binding.ValueType )
+			TArray<UObject*> ValueTypeValidationData;
+			PopulateValueTypeValidationData( PropetyAtCurrentLevel, ValueTypeValidationData );
+			// Perf Note: We might be able to cache this validation and some part of this function at some point
+			if ( ValueTypeValidationData == Binding.ValueTypeValidationData )
 			{
 				return PropetyAtCurrentLevel;
 			}
+			else
+			{
+				UE_LOG( LogDataprepCore, Warning, TEXT("A binding was invalid because it's type changed") );
+				OutPropertyValueAddress = nullptr;
+			}
 		}
 
 		return nullptr;
@@ -301,7 +455,7 @@ namespace DataprepParameterization
 	void CopyCompleteValue(UProperty& DestinationProperty, void* DestinationAddress, UProperty& SourceProperty, void* SourceAddress)
 	{
 		UClass* ProperyClass = DestinationProperty.GetClass();
-		// We only support copying value of properties when they are from the same class
+		// We only support copying value of properties when they are from the same class (this is not a warranty that this is safe, it's only a validation heuristic)
 		check( ProperyClass == SourceProperty.GetClass() );
 
 		// Bool properties are special because each property can have their own mask and offset from there base address (probably to support bitfields)
@@ -318,18 +472,31 @@ namespace DataprepParameterization
 };
 
 
-bool FDataprepParameterizationBinding::operator==(const FDataprepParameterizationBinding& Other) const
+FDataprepParameterizationBinding::FDataprepParameterizationBinding(UDataprepParameterizableObject* InObjectBinded, TArray<FDataprepPropertyLink> InPropertyChain)
+	: ObjectBinded( InObjectBinded )
+	, PropertyChain( MoveTemp( InPropertyChain ) )
+	, ValueTypeValidationData()
 {
-	// Not a accurate but ok for now
-	return ObjectBinded == Other.ObjectBinded && GetTypeHash(PropertyChain) == GetTypeHash(Other.PropertyChain) && ValueType == Other.ValueType;
+	if ( PropertyChain.Num() > 0 )
+	{
+		if ( UProperty* Property = PropertyChain.Last().CachedProperty.Get() )
+		{
+			DataprepParameterization::PopulateValueTypeValidationData( Property, ValueTypeValidationData );
+		}
+	}
 }
 
+bool FDataprepParameterizationBinding::operator==(const FDataprepParameterizationBinding& Other) const
+{
+	// The value type validation data shouldn't matter when comparing binding
+	return ObjectBinded == Other.ObjectBinded && PropertyChain == Other.PropertyChain;
+}
 
 uint32 GetTypeHash(const FDataprepParameterizationBinding& Binding)
 {
-	uint32 Hash = GetTypeHash( Binding.ObjectBinded );
-	Hash = HashCombine( Hash, GetTypeHash( Binding.PropertyChain ) );
-	return HashCombine( Hash, GetTypeHash( Binding.ValueType ) );
+	// The value type validation data shouldn't matter for the hash of a binding
+	return HashCombine(  GetTypeHash( Binding.ObjectBinded ), GetTypeHash( Binding.PropertyChain ) );
+	
 }
 
 uint32 GetTypeHash(const TArray<FDataprepPropertyLink>& PropertyLinks)
@@ -488,7 +655,7 @@ void UDataprepParameterizationBindings::AddReferencedObjects(UObject* InThis, FR
 		for ( TSharedRef<FDataprepParameterizationBinding>& Binding : BindingSet )
 		{
 			Collector.AddReferencedObject( Binding->ObjectBinded );
-			Collector.AddReferencedObject( Binding->ValueType );
+			Collector.AddReferencedObjects( Binding->ValueTypeValidationData );
 		}
 	}
 
@@ -723,8 +890,10 @@ bool UDataprepParameterization::BindObjectProperty(UDataprepParameterizableObjec
 			{
 				UProperty* Property = *PropertyPtr;
 
-				// Ensure that the properties type match
-				if ( Property->GetClass() == Binding->ValueType )
+				// Ensure that the properties are compatible 
+				TArray<UObject*> ValueTypeValidationData;
+				DataprepParameterization::PopulateValueTypeValidationData( Property, ValueTypeValidationData );
+				if ( ValueTypeValidationData == Binding->ValueTypeValidationData )
 				{
 					Modify();
 					BindingsContainer->Add( Binding, Name );
@@ -872,20 +1041,30 @@ void UDataprepParameterization::OnObjectPostEdit(UDataprepParameterizableObject*
 	}
 }
 
-void UDataprepParameterization::GetExistingParameterNamesForType(UClass* PropertyClass, TSet<FString>& OutValidExistingNames, TSet<FString>& OutInvalidNames) const
+void UDataprepParameterization::GetExistingParameterNamesForType(UProperty* Property, TSet<FString>& OutValidExistingNames, TSet<FString>& OutInvalidNames) const
 {
 	OutValidExistingNames.Empty( NameToParameterizationProperty.Num() );
 	OutInvalidNames.Empty( NameToParameterizationProperty.Num() );
 
 	for ( const TPair<FName, UProperty*>& Pair : NameToParameterizationProperty )
 	{
-		if ( const UProperty* Property = Pair.Value )
+		if ( UProperty* ParameterizationProperty = Pair.Value )
 		{
-			if ( Property->GetClass() == PropertyClass )
+			bool bWasAdded = false;
+			if ( ParameterizationProperty->GetClass() == Property->GetClass() )
 			{
-				OutValidExistingNames.Add( Pair.Key.ToString() );
+				TArray<UObject*> ValidationDataForParameterizationProperty;
+				DataprepParameterization::PopulateValueTypeValidationData( ParameterizationProperty, ValidationDataForParameterizationProperty );
+				TArray<UObject*> ValidationData;
+				DataprepParameterization::PopulateValueTypeValidationData( Property, ValidationData );
+				if ( ValidationDataForParameterizationProperty == ValidationData )
+				{
+					bWasAdded = true;
+					OutValidExistingNames.Add( Pair.Key.ToString() );
+				}
 			}
-			else
+			
+			if ( !bWasAdded )
 			{
 				OutInvalidNames.Add( Pair.Key.ToString() );
 			}
diff --git a/Engine/Plugins/Experimental/DataPrepEditor/Source/DataPrepCore/Private/Parameterization/DataprepParameterization.h b/Engine/Plugins/Experimental/DataPrepEditor/Source/DataPrepCore/Private/Parameterization/DataprepParameterization.h
--- a/Engine/Plugins/Experimental/DataPrepEditor/Source/DataPrepCore/Private/Parameterization/DataprepParameterization.h
+++ b/Engine/Plugins/Experimental/DataPrepEditor/Source/DataPrepCore/Private/Parameterization/DataprepParameterization.h
@@ -18,6 +18,10 @@
 class UDataprepParameterizableObject;
 class UProperty;
 
+/**
+ * The parameterization binding is a struct that hold an object and the property path to the parameterized property
+ * It also hold a array to validate the that value type of the parameterized property didn't change since it's creation
+ */
 USTRUCT()
 struct FDataprepParameterizationBinding
 {
@@ -26,22 +30,10 @@ struct FDataprepParameterizationBinding
 	FDataprepParameterizationBinding()
 		: ObjectBinded( nullptr )
 		, PropertyChain()
-		, ValueType( nullptr )
+		, ValueTypeValidationData()
 	{}
 
-	FDataprepParameterizationBinding(UDataprepParameterizableObject* InObjectBinded, TArray<FDataprepPropertyLink> InPropertyChain)
-		: ObjectBinded( InObjectBinded )
-		, PropertyChain( MoveTemp( InPropertyChain ) )
-		, ValueType( nullptr )
-	{
-		if ( PropertyChain.Num() > 0 )
-		{
-			if ( UProperty* Property = PropertyChain.Last().CachedProperty.Get() )
-			{
-				ValueType = Property->GetClass();
-			}
-		}
-	}
+	FDataprepParameterizationBinding(UDataprepParameterizableObject* InObjectBinded, TArray<FDataprepPropertyLink> InPropertyChain);
 
 	FDataprepParameterizationBinding(FDataprepParameterizationBinding&&) = default;
 	FDataprepParameterizationBinding(const FDataprepParameterizationBinding&) = default;
@@ -56,9 +48,9 @@ struct FDataprepParameterizationBinding
 	UPROPERTY()
 	TArray<FDataprepPropertyLink> PropertyChain;
 
-	// The class of the property managing the value
+	// Value Type Validation Array. This is the result of a depth first search on the parametrized property
 	UPROPERTY()
-	UClass* ValueType;
+	TArray<UObject*> ValueTypeValidationData;
 };
 
 uint32 GetTypeHash(const FDataprepParameterizationBinding& Binding);
@@ -220,7 +212,7 @@ public:
 
 	void OnObjectPostEdit(UDataprepParameterizableObject* Object, const TArray<FDataprepPropertyLink>& PropertyChain, EPropertyChangeType::Type ChangeType);
 
-	void GetExistingParameterNamesForType(UClass* PropertyClass,TSet<FString>& OutValidExistingNames, TSet<FString>& OutInvalidNames) const;
+	void GetExistingParameterNamesForType(UProperty* Property,TSet<FString>& OutValidExistingNames, TSet<FString>& OutInvalidNames) const;
 
 private:
 
diff --git a/Engine/Plugins/Experimental/DataPrepEditor/Source/DataPrepCore/Private/Parameterization/DataprepParameterizationUtils.cpp b/Engine/Plugins/Experimental/DataPrepEditor/Source/DataPrepCore/Private/Parameterization/DataprepParameterizationUtils.cpp
--- a/Engine/Plugins/Experimental/DataPrepEditor/Source/DataPrepCore/Private/Parameterization/DataprepParameterizationUtils.cpp
+++ b/Engine/Plugins/Experimental/DataPrepEditor/Source/DataPrepCore/Private/Parameterization/DataprepParameterizationUtils.cpp
@@ -40,6 +40,11 @@ namespace DataprepParameterizationUtils
 
 }
 
+bool operator==(const FDataprepPropertyLink& A, const FDataprepPropertyLink& B)
+{
+	return A.PropertyName == B.PropertyName && A.ContainerIndex == B.ContainerIndex;
+}
+
 uint32 GetTypeHash(const FDataprepPropertyLink& PropertyLink)
 {
 	return HashCombine( GetTypeHash( PropertyLink.PropertyName ), GetTypeHash( PropertyLink.ContainerIndex ) );
diff --git a/Engine/Plugins/Experimental/DataPrepEditor/Source/DataPrepCore/Public/DataPrepAsset.h b/Engine/Plugins/Experimental/DataPrepEditor/Source/DataPrepCore/Public/DataPrepAsset.h
--- a/Engine/Plugins/Experimental/DataPrepEditor/Source/DataPrepCore/Public/DataPrepAsset.h
+++ b/Engine/Plugins/Experimental/DataPrepEditor/Source/DataPrepCore/Public/DataPrepAsset.h
@@ -113,7 +113,7 @@ public:
 
 	void RemoveObjectPropertyFromParameterization(UDataprepParameterizableObject* Object, const TArray<struct FDataprepPropertyLink>& InPropertyChain);
 
-	void GetExistingParameterNamesForType(UClass* PropertyClass, TSet<FString>& OutValidExistingNames, TSet<FString>& OutInvalidNames) const;
+	void GetExistingParameterNamesForType(UProperty* Property, TSet<FString>& OutValidExistingNames, TSet<FString>& OutInvalidNames) const;
 
 	UDataprepParameterization* GetDataprepParameterization() { return Parameterization; }
 
diff --git a/Engine/Plugins/Experimental/DataPrepEditor/Source/DataPrepCore/Public/Parameterization/DataprepParameterizationUtils.h b/Engine/Plugins/Experimental/DataPrepEditor/Source/DataPrepCore/Public/Parameterization/DataprepParameterizationUtils.h
--- a/Engine/Plugins/Experimental/DataPrepEditor/Source/DataPrepCore/Public/Parameterization/DataprepParameterizationUtils.h
+++ b/Engine/Plugins/Experimental/DataPrepEditor/Source/DataPrepCore/Public/Parameterization/DataprepParameterizationUtils.h
@@ -47,6 +47,8 @@ struct FDataprepPropertyLink
 	FDataprepPropertyLink& operator=(const FDataprepPropertyLink& Other) = default;
 	FDataprepPropertyLink& operator=(FDataprepPropertyLink&& Other) = default;
 
+	friend bool operator==(const FDataprepPropertyLink& A,const FDataprepPropertyLink& B);
+
 	UPROPERTY()
 	TWeakObjectPtr<UProperty> CachedProperty;
 
diff --git a/Engine/Plugins/Experimental/DataPrepEditor/Source/DataPrepEditor/Private/Widgets/Parameterization/SDataprepLinkToParameter.cpp b/Engine/Plugins/Experimental/DataPrepEditor/Source/DataPrepEditor/Private/Widgets/Parameterization/SDataprepLinkToParameter.cpp
--- a/Engine/Plugins/Experimental/DataPrepEditor/Source/DataPrepEditor/Private/Widgets/Parameterization/SDataprepLinkToParameter.cpp
+++ b/Engine/Plugins/Experimental/DataPrepEditor/Source/DataPrepEditor/Private/Widgets/Parameterization/SDataprepLinkToParameter.cpp
@@ -34,8 +34,7 @@ void SDataprepLinkToParameter::Construct(const FArguments& InArgs, const TShared
 	{
 		if ( UProperty* Property = ParameterizationActionData->PropertyChain.Last().CachedProperty.Get() )
 		{
-			UClass* PropertyClass = Property->GetClass();
-			ParameterizationActionData->DataprepAsset->GetExistingParameterNamesForType( PropertyClass, ValidExistingNames, InvalidNames );
+			ParameterizationActionData->DataprepAsset->GetExistingParameterNamesForType( Property, ValidExistingNames, InvalidNames );
 
 			FString SuggestedName = ParameterizationActionData->PropertyChain.Last().PropertyName.ToString();
 
