Added a stronger validation of the property chain for the dataprep parameterization.

diff --git a/Engine/Plugins/Experimental/DataPrepEditor/Source/DataPrepCore/Private/Parameterization/DataprepParameterization.cpp b/Engine/Plugins/Experimental/DataPrepEditor/Source/DataPrepCore/Private/Parameterization/DataprepParameterization.cpp
--- a/Engine/Plugins/Experimental/DataPrepEditor/Source/DataPrepCore/Private/Parameterization/DataprepParameterization.cpp
+++ b/Engine/Plugins/Experimental/DataPrepEditor/Source/DataPrepCore/Private/Parameterization/DataprepParameterization.cpp
@@ -5,13 +5,13 @@
 #include "DataPrepAsset.h"
 #include "DataprepParameterizableObject.h"
 #include "DataprepParameterizationArchive.h"
+#include "Parameterization/DataprepParameterizationUtils.h"
 
 #include "CoreGlobals.h"
 #include "DataprepCoreLogCategory.h"
 #include "Editor.h"
 #include "Engine/Engine.h"
 #include "Math/UnrealMathUtility.h"
-#include "Parameterization/DataprepParameterizationUtils.h"
 #include "Serialization/Archive.h"
 #include "Templates/UnrealTemplate.h"
 #include "UObject/Object.h"
@@ -864,7 +864,7 @@ UObject* UDataprepParameterization::GetDefaultObject()
 
 bool UDataprepParameterization::BindObjectProperty(UDataprepParameterizableObject* Object, const TArray<FDataprepPropertyLink>& PropertyChain, const FName& Name)
 {
-	if ( Object && PropertyChain.Num() > 0 && !Name.IsNone() )
+	if ( Object && FDataprepParameterizationUtils::IsPropertyChainValid( PropertyChain ) && !Name.IsNone() )
 	{
 		TSharedRef<FDataprepParameterizationBinding> Binding = MakeShared<FDataprepParameterizationBinding>( Object, PropertyChain );
 		void* AddressOfTheValueFromBinding;
diff --git a/Engine/Plugins/Experimental/DataPrepEditor/Source/DataPrepCore/Private/Parameterization/DataprepParameterizationUtils.cpp b/Engine/Plugins/Experimental/DataPrepEditor/Source/DataPrepCore/Private/Parameterization/DataprepParameterizationUtils.cpp
--- a/Engine/Plugins/Experimental/DataPrepEditor/Source/DataPrepCore/Private/Parameterization/DataprepParameterizationUtils.cpp
+++ b/Engine/Plugins/Experimental/DataPrepEditor/Source/DataPrepCore/Private/Parameterization/DataprepParameterizationUtils.cpp
@@ -150,7 +150,7 @@ FDataprepParameterizationContext FDataprepParameterizationUtils::CreateContext(T
 	{
 		// This implementation could be improved by incrementally expending the property chain.
 		NewContext.PropertyChain = MakePropertyChain( PropertyHandle );
-		NewContext.State = NewContext.PropertyChain.Num() > 0 ? EParametrizationState::CanBeParameterized : EParametrizationState::InvalidForParameterization;
+		NewContext.State = IsPropertyChainValid( NewContext.PropertyChain ) ? EParametrizationState::CanBeParameterized : EParametrizationState::InvalidForParameterization;
 	}
 	else if ( NewContext.State == EParametrizationState::IsParameterized )
 	{
@@ -195,4 +195,47 @@ UDataprepAsset* FDataprepParameterizationUtils::GetDataprepAssetForParameterizat
 	}
 
 	return nullptr;
-}
\ No newline at end of file
+}
+
+bool FDataprepParameterizationUtils::IsPropertyChainValid(const TArray<FDataprepPropertyLink>& PropertyChain)
+{
+	if ( PropertyChain.Num() == 0 )
+	{
+		return false;
+	}
+
+	// We will support those latter (there is still some bugs to fix before)
+	if ( PropertyChain.Num() > 1 )
+	{
+		return false;
+	}
+
+
+	for ( const FDataprepPropertyLink& PropertyLink : PropertyChain )
+	{
+		UProperty* Property = PropertyLink.CachedProperty.Get();
+		if ( !Property )
+		{
+			return false;
+		}
+
+		// todo check for custom getter and setter
+
+		// Temporary unsupported property check
+		{
+			// Reject container properties for now (there is still some bugs to fix before)
+			if ( DataprepParameterizationUtils::IsAContainerProperty( Property ) )
+			{
+				return false;
+			}
+
+			// We are not able to serialize the text properties yet
+			if ( Property->GetClass() == UTextProperty::StaticClass() )
+			{
+				return false;
+			}
+		}
+	}
+
+	return true;
+}
diff --git a/Engine/Plugins/Experimental/DataPrepEditor/Source/DataPrepCore/Public/Parameterization/DataprepParameterizationUtils.h b/Engine/Plugins/Experimental/DataPrepEditor/Source/DataPrepCore/Public/Parameterization/DataprepParameterizationUtils.h
--- a/Engine/Plugins/Experimental/DataPrepEditor/Source/DataPrepCore/Public/Parameterization/DataprepParameterizationUtils.h
+++ b/Engine/Plugins/Experimental/DataPrepEditor/Source/DataPrepCore/Public/Parameterization/DataprepParameterizationUtils.h
@@ -89,4 +89,10 @@ public:
 	 * @return A valid pointer if the object was valid for the parameterization
 	 */
 	static UDataprepAsset* GetDataprepAssetForParameterization(UObject* Object);
+
+	/**
+	 * Check if a property is chain supported by the dataprep parameterization systems
+	 * Note that the cached properties of the chain must not be null
+	 */
+	static bool IsPropertyChainValid(const TArray<FDataprepPropertyLink>& PropertyChain);
 };
