Cleaned up UI of Dataprep editor when editing a Dataprep instance asset
Improved the UI for'externalized' parameters
Added exposition of Parent asset for Dataprep instance

diff --git a/Engine/Plugins/Experimental/DataPrepEditor/Source/DataPrepEditor/Private/DataPrepEditor.cpp b/Engine/Plugins/Experimental/DataPrepEditor/Source/DataPrepEditor/Private/DataPrepEditor.cpp
--- a/Engine/Plugins/Experimental/DataPrepEditor/Source/DataPrepEditor/Private/DataPrepEditor.cpp
+++ b/Engine/Plugins/Experimental/DataPrepEditor/Source/DataPrepEditor/Private/DataPrepEditor.cpp
@@ -4,6 +4,7 @@
 
 #include "DataPrepOperation.h"
 
+#include "DataprepAssetInstance.h"
 #include "DataPrepContentConsumer.h"
 #include "DataPrepContentProducer.h"
 #include "DataPrepEditorActions.h"
@@ -144,6 +145,7 @@ FDataprepEditor::FDataprepEditor()
 	: bWorldBuilt(false)
 	, bIsFirstRun(false)
 	, bPipelineChanged(false)
+	, bIsDataprepInstance(false)
 	, bIsActionMenuContextSensitive(true)
 	, bSaveIntermediateBuildProducts(false)
 	, PreviewWorld(nullptr)
@@ -284,11 +286,6 @@ void FDataprepEditor::RegisterTabSpawners(const TSharedRef<class FTabManager>& I
 		.SetGroup(WorkspaceMenuCategoryRef)
 		.SetIcon(FSlateIcon(FDataprepEditorStyle::GetStyleSetName(), "DataprepEditor.Tabs.SceneViewport"));
 
-	InTabManager->RegisterTabSpawner(PaletteTabId, FOnSpawnTab::CreateSP(this, &FDataprepEditor::SpawnTabPalette))
-		.SetDisplayName(LOCTEXT("PaletteTab", "Palette"))
-		.SetGroup(WorkspaceMenuCategoryRef)
-		.SetIcon( FSlateIcon(FEditorStyle::GetStyleSetName(), "Kismet.Tabs.Palette"));
-
 	InTabManager->RegisterTabSpawner(DetailsTabId, FOnSpawnTab::CreateSP(this, &FDataprepEditor::SpawnTabDetails))
 		.SetDisplayName(LOCTEXT("DetailsTab", "Details"))
 		.SetGroup(WorkspaceMenuCategoryRef)
@@ -304,12 +301,21 @@ void FDataprepEditor::RegisterTabSpawners(const TSharedRef<class FTabManager>& I
 		.SetGroup(WorkspaceMenuCategoryRef)
 		.SetIcon(FSlateIcon(FEditorStyle::GetStyleSetName(), "LevelEditor.Tabs.StatsViewer"));
 
-	// Temp code for the nodes development
-	InTabManager->RegisterTabSpawner(PipelineGraphTabId, FOnSpawnTab::CreateSP(this, &FDataprepEditor::SpawnTabPipelineGraph))
-		.SetDisplayName(LOCTEXT("PipelineGraphTab", "Pipeline Graph"))
-		.SetGroup(WorkspaceMenuCategoryRef)
-		.SetIcon(FSlateIcon(FEditorStyle::GetStyleSetName(), "GraphEditor.EventGraph_16x"));
-	// end of temp code for nodes development
+	// Do not register tabs which are not pertinent to Dataprep instance
+	if(!bIsDataprepInstance)
+	{
+		InTabManager->RegisterTabSpawner(PaletteTabId, FOnSpawnTab::CreateSP(this, &FDataprepEditor::SpawnTabPalette))
+			.SetDisplayName(LOCTEXT("PaletteTab", "Palette"))
+			.SetGroup(WorkspaceMenuCategoryRef)
+			.SetIcon( FSlateIcon(FEditorStyle::GetStyleSetName(), "Kismet.Tabs.Palette"));
+
+		// Temp code for the nodes development
+		InTabManager->RegisterTabSpawner(PipelineGraphTabId, FOnSpawnTab::CreateSP(this, &FDataprepEditor::SpawnTabPipelineGraph))
+			.SetDisplayName(LOCTEXT("PipelineGraphTab", "Pipeline Graph"))
+			.SetGroup(WorkspaceMenuCategoryRef)
+			.SetIcon(FSlateIcon(FEditorStyle::GetStyleSetName(), "GraphEditor.EventGraph_16x"));
+		// end of temp code for nodes development
+	}
 }
 
 void FDataprepEditor::UnregisterTabSpawners(const TSharedRef<class FTabManager>& InTabManager)
@@ -387,6 +393,8 @@ void FDataprepEditor::InitDataprepEditor(const EToolkitMode::Type Mode, const TS
 	DataprepAssetInterfacePtr = TWeakObjectPtr<UDataprepAssetInterface>(InDataprepAssetInterface);
 	check( DataprepAssetInterfacePtr.IsValid() );
 
+	bIsDataprepInstance = Cast<UDataprepAssetInstance>(InDataprepAssetInterface) != nullptr;
+
 	DataprepAssetInterfacePtr->GetOnChanged().AddSP( this, &FDataprepEditor::OnDataprepAssetChanged );
 
 	// Assign unique session identifier
@@ -838,12 +846,17 @@ TSharedRef<SDockTab> FDataprepEditor::SpawnTabPipelineGraph(const FSpawnTabArgs
 {
 	check(Args.GetTabId() == PipelineGraphTabId);
 
-	return SNew(SDockTab)
-		//.Icon(FDataprepEditorStyle::Get()->GetBrush("DataprepEditor.Tabs.Pipeline"))
-		.Label(LOCTEXT("DataprepEditor_PipelineTab_Title", "Pipeline"))
-		[
-			DataprepRecipeBPPtr.IsValid() ? PipelineView.ToSharedRef() : SNullWidget::NullWidget
-		];
+	if(!bIsDataprepInstance)
+	{
+		return SNew(SDockTab)
+			//.Icon(FDataprepEditorStyle::Get()->GetBrush("DataprepEditor.Tabs.Pipeline"))
+			.Label(LOCTEXT("DataprepEditor_PipelineTab_Title", "Pipeline"))
+			[
+				DataprepRecipeBPPtr.IsValid() ? PipelineView.ToSharedRef() : SNullWidget::NullWidget
+			];
+	}
+
+	return SNew(SDockTab);
 }
 // end of temp code for nodes development
 
@@ -880,12 +893,17 @@ TSharedRef<SDockTab> FDataprepEditor::SpawnTabPalette(const FSpawnTabArgs & Args
 {
 	check(Args.GetTabId() == PaletteTabId);
 
-	return SNew(SDockTab)
-		.Icon(FSlateIcon(FEditorStyle::GetStyleSetName(), "Kismet.Tabs.Palette").GetIcon())
-		.Label(LOCTEXT("PaletteTab", "Palette"))
-		[
-			SNew(SDataprepPalette)
-		];
+	if(!bIsDataprepInstance)
+	{
+		return SNew(SDockTab)
+			.Icon(FSlateIcon(FEditorStyle::GetStyleSetName(), "Kismet.Tabs.Palette").GetIcon())
+			.Label(LOCTEXT("PaletteTab", "Palette"))
+			[
+				SNew(SDataprepPalette)
+			];
+	}
+
+	return SNew(SDockTab);
 }
 
 TSharedRef<SDockTab> FDataprepEditor::SpawnTabDataprep(const FSpawnTabArgs & Args)
diff --git a/Engine/Plugins/Experimental/DataPrepEditor/Source/DataPrepEditor/Private/DataPrepEditor.h b/Engine/Plugins/Experimental/DataPrepEditor/Source/DataPrepEditor/Private/DataPrepEditor.h
--- a/Engine/Plugins/Experimental/DataPrepEditor/Source/DataPrepEditor/Private/DataPrepEditor.h
+++ b/Engine/Plugins/Experimental/DataPrepEditor/Source/DataPrepEditor/Private/DataPrepEditor.h
@@ -189,6 +189,7 @@ private:
 	bool bWorldBuilt;
 	bool bIsFirstRun;
 	bool bPipelineChanged;
+	bool bIsDataprepInstance;
 	TWeakObjectPtr<UDataprepAssetInterface> DataprepAssetInterfacePtr;
 
 	FOnDataprepAssetProducerChanged DataprepAssetProducerChangedDelegate;
diff --git a/Engine/Plugins/Experimental/DataPrepEditor/Source/DataPrepEditor/Private/DataprepEditorStyle.cpp b/Engine/Plugins/Experimental/DataPrepEditor/Source/DataPrepEditor/Private/DataprepEditorStyle.cpp
--- a/Engine/Plugins/Experimental/DataPrepEditor/Source/DataPrepEditor/Private/DataprepEditorStyle.cpp
+++ b/Engine/Plugins/Experimental/DataPrepEditor/Source/DataPrepEditor/Private/DataprepEditorStyle.cpp
@@ -13,6 +13,8 @@
 
 #define IMAGE_PLUGIN_BRUSH( RelativePath, ... ) FSlateImageBrush( FDataprepEditorStyle::InContent( RelativePath, ".png" ), __VA_ARGS__ )
 
+#define DEFAULT_FONT(...) FCoreStyle::GetDefaultFontStyle(__VA_ARGS__)
+
 TSharedPtr<FSlateStyleSet> FDataprepEditorStyle::StyleSet;
 
 void FDataprepEditorStyle::Initialize()
@@ -57,6 +59,8 @@ void FDataprepEditorStyle::Initialize()
 	StyleSet->Set("DataprepEditor.ExecutePipeline.Selected", new IMAGE_PLUGIN_BRUSH("Icons/ExecutePipeline", Icon40x40));
 	StyleSet->Set("DataprepEditor.ExecutePipeline.Selected.Small", new IMAGE_PLUGIN_BRUSH("Icons/ExecutePipeline", Icon20x20));
 
+	StyleSet->Set( "DataprepEditor.SectionFont", DEFAULT_FONT( "Bold", 10 ) );
+
 	// Dataprep action UI
 	{
 		StyleSet->Set("DataprepAction.Padding", 2.f);
diff --git a/Engine/Plugins/Experimental/DataPrepEditor/Source/DataPrepEditor/Private/Widgets/DataprepAssetView.cpp b/Engine/Plugins/Experimental/DataPrepEditor/Source/DataPrepEditor/Private/Widgets/DataprepAssetView.cpp
--- a/Engine/Plugins/Experimental/DataPrepEditor/Source/DataPrepEditor/Private/Widgets/DataprepAssetView.cpp
+++ b/Engine/Plugins/Experimental/DataPrepEditor/Source/DataPrepEditor/Private/Widgets/DataprepAssetView.cpp
@@ -4,6 +4,7 @@
 
 #include "DataPrepAsset.h"
 #include "DataPrepEditor.h"
+#include "DataprepAssetInstance.h"
 
 #include "DataPrepContentConsumer.h"
 #include "DataprepEditorStyle.h"
@@ -360,7 +361,67 @@ void SDataprepAssetView::Construct( const FArguments& InArgs, UDataprepAssetInte
 
 	TSharedRef<SScrollBar> ScrollBar = SNew(SScrollBar);
 
-	// #ueent_todo: Look at changing the border brushes to add color to this stuff
+	TSharedRef<SWidget> ParentWidget = SNullWidget::NullWidget;
+	TSharedRef<SWidget> ParentSpacer = SNullWidget::NullWidget;
+	if(UDataprepAssetInstance* DataprepInstance = Cast<UDataprepAssetInstance>(InDataprepAssetPtr))
+	{
+		ParentWidget = SNew(SHorizontalBox)
+		+SHorizontalBox::Slot()
+		.FillWidth(1)
+		[
+			SNew(SVerticalBox)
+			+ SVerticalBox::Slot()
+			.AutoHeight()
+			[
+				SNew(SHorizontalBox)
+				+ SHorizontalBox::Slot()
+				[
+					SNew(SSpacer)
+					.Size( FVector2D( 200, 10 ) )		
+				]
+			]
+			+ SVerticalBox::Slot()
+			//.Padding( 5.0f )
+			.AutoHeight()
+			//.MaxHeight( 400.f )
+			[
+				DataprepWidgetUtils::CreateParameterRow( SNew(SDataprepInstanceParentWidget).ColumnSizeData(ColumnSizeData).DataprepInstance(DataprepInstance) )
+			]
+		];
+
+		ParentSpacer = SNew(SHorizontalBox)
+		+SHorizontalBox::Slot()
+		.FillWidth(1)
+		[
+			SNew(SVerticalBox)
+			+ SVerticalBox::Slot()
+			.AutoHeight()
+			[
+				SNew(SHorizontalBox)
+				+ SHorizontalBox::Slot()
+				[
+					SNew(SSpacer)
+					.Size( FVector2D( 200, 10 ) )		
+				]
+			]
+			+ SVerticalBox::Slot()
+			.AutoHeight()
+			[
+				SNew(SHorizontalBox)
+				+ SHorizontalBox::Slot()
+				.FillWidth(1.0f)
+				.Padding( 0, 10, 0, 0 )
+				.HAlign( EHorizontalAlignment::HAlign_Center )
+				[
+					// #ueent_todo: make color block's width vary with parent widget
+					SNew(SColorBlock)
+					.Color( FLinearColor( 0.9f, 0.9f, 0.9f ) )
+					.Size( FVector2D( 3000, 1 ) )
+				]
+			]
+		];
+	}
+
 	ChildSlot
 	[
 		SNew(SBorder)
@@ -377,6 +438,19 @@ void SDataprepAssetView::Construct( const FArguments& InArgs, UDataprepAssetInte
 				+ SScrollBox::Slot()
 				[
 					SNew(SVerticalBox)
+					// Begin - Section for Dataprep parent
+					+ SVerticalBox::Slot()
+					.Padding( 5.0f )
+					.AutoHeight()
+					.MaxHeight( 400.f )
+					[
+						ParentWidget
+					]
+					+ SVerticalBox::Slot()
+					[
+						ParentSpacer
+					]
+					// End - Section for Dataprep parent
 					// Begin - Section for producers
 					+ SVerticalBox::Slot()
 					.Padding( 5.0f )
@@ -426,7 +500,7 @@ void SDataprepAssetView::Construct( const FArguments& InArgs, UDataprepAssetInte
 							SNew(STextBlock)
 							.Text(LOCTEXT("DataprepAssetView_Consumer_label", "Output"))
 							.MinDesiredWidth( 200 )
-							.Font( IDetailLayoutBuilder::GetDetailFontBold() )
+							.Font( FDataprepEditorStyle::Get()->GetFontStyle("DataprepEditor.SectionFont") )
 						]
 						+ SHorizontalBox::Slot()
 						.AutoWidth()
@@ -499,23 +573,27 @@ void SDataprepAssetView::Construct( const FArguments& InArgs, UDataprepAssetInte
 							SNew(STextBlock)
 							.Text(LOCTEXT("DataprepAssetView_Consumer_Parameterization", "Parameterization"))
 							.MinDesiredWidth( 200 )
-							.Font( IDetailLayoutBuilder::GetDetailFontBold() )
+							.Font( FDataprepEditorStyle::Get()->GetFontStyle("DataprepEditor.SectionFont") )
 						]
 					]
 					+ SVerticalBox::Slot()
 					.AutoHeight()
-					.Padding(10.0f, 5.0f, 0.0f, 5.0f)
 					[
 						SNew(SHorizontalBox)
 						+ SHorizontalBox::Slot()
 						.FillWidth(1.0f)
 						.HAlign(EHorizontalAlignment::HAlign_Left)
 						.VAlign(VAlign_Center)
-						.Padding( FMargin( -5.0f, 5.0f, 20.0f, 5.0f ) )
+						.Padding(10.0f, 5.0f, 0.0f, 5.0f)
 						[
-							SNew( SDataprepDetailsView )
-							.Object( DataprepAssetInterfacePtr->GetParameterizationObject() )
-							.ColumnSizeData( ColumnSizeData )
+							SNew( SConstrainedBox )
+							[
+								SNew( SDataprepDetailsView )
+								.Object( DataprepAssetInterfacePtr->GetParameterizationObject() )
+								.ColumnSizeData( ColumnSizeData )
+								.Spacing( 10.0f )
+								.ColumnPadding( true )
+							]
 						]
 					]
 					// End - Section for parameterization
diff --git a/Engine/Plugins/Experimental/DataPrepEditor/Source/DataPrepEditor/Private/Widgets/DataprepWidgets.cpp b/Engine/Plugins/Experimental/DataPrepEditor/Source/DataPrepEditor/Private/Widgets/DataprepWidgets.cpp
--- a/Engine/Plugins/Experimental/DataPrepEditor/Source/DataPrepEditor/Private/Widgets/DataprepWidgets.cpp
+++ b/Engine/Plugins/Experimental/DataPrepEditor/Source/DataPrepEditor/Private/Widgets/DataprepWidgets.cpp
@@ -3,6 +3,7 @@
 #include "DataprepWidgets.h"
 
 #include "DataPrepAsset.h"
+#include "DataprepAssetInstance.h"
 #include "DataPrepContentConsumer.h"
 #include "DataprepAssetView.h"
 #include "DataprepEditorLogCategory.h"
@@ -51,32 +52,50 @@
 
 #define LOCTEXT_NAMESPACE "DataprepSlateHelper"
 
-/** Helper class to force a widget to fill in a space. Copied from SDetailSingleItemRow.cpp */
-class SConstrainedBox : public SCompoundWidget
+namespace DataprepWidgetUtils
 {
-public:
-	SLATE_BEGIN_ARGS(SConstrainedBox)
+	TSharedRef<SWidget> CreatePropertyWidget( TSharedPtr<SWidget> NameWidget, TSharedPtr<SWidget> ValueWidget, TSharedPtr< FDataprepDetailsViewColumnSizeData > ColumnSizeData, float Spacing, EHorizontalAlignment HAlign = EHorizontalAlignment::HAlign_Left, EVerticalAlignment VAlign = EVerticalAlignment::VAlign_Center )
 	{
-	}
-	SLATE_DEFAULT_SLOT(FArguments, Content)
-	SLATE_END_ARGS()
-
-	void Construct(const FArguments& InArgs)
-	{
-		ChildSlot
+		return SNew(SHorizontalBox)
+		+ SHorizontalBox::Slot()
+		.FillWidth(1.f)
+		.VAlign(VAlign_Center)
+		.HAlign(HAlign_Left)
+		.Padding(0.0f, 0.0f, 0.0f, Spacing)
 		[
-			InArgs._Content.Widget
+			SNew(SSplitter)
+			.Style(FEditorStyle::Get(), "DetailsView.Splitter")
+			.PhysicalSplitterHandleSize(1.0f)
+			.HitDetectionSplitterHandleSize(5.0f)
+			.ResizeMode( ESplitterResizeMode::Fill )
+			+ SSplitter::Slot()
+			.Value(ColumnSizeData->LeftColumnWidth)
+			.OnSlotResized( SSplitter::FOnSlotResized::CreateLambda( [](float InNewWidth) -> void {} ) )
+			[
+				NameWidget.ToSharedRef()
+			]
+			+ SSplitter::Slot()
+			.Value(ColumnSizeData->RightColumnWidth)
+			.OnSlotResized(ColumnSizeData->OnWidthChanged)
+			[
+				SNew(SHorizontalBox)
+				.Clipping(EWidgetClipping::OnDemand)
+				+ SHorizontalBox::Slot()
+				.Padding( 5.0f, 2.5f, 2.0f, 2.5f )
+				.HAlign( HAlign )
+				.VAlign( VAlign )
+				[
+					// Trick to force the splitter widget to fill up the space of its parent
+					// Strongly inspired from SDetailSingleItemRow
+					SNew(SConstrainedBox)
+					[
+						ValueWidget.ToSharedRef()
+					]
+				]
+			]
 		];
 	}
-
-	virtual FVector2D ComputeDesiredSize(float LayoutScaleMultiplier) const override
-	{
-		// Voluntarily ridiculously large value to force the child widget to fill up the available space
-		const float MinWidthVal = 2000;
-		const FVector2D ChildSize = ChildSlot.GetWidget()->GetDesiredSize();
-		return FVector2D(FMath::Max(MinWidthVal, ChildSize.X), ChildSize.Y);
-	}
-};
+}
 
 void SDataprepConsumerWidget::OnLevelNameChanged( const FText &NewLevelName, ETextCommit::Type CommitType )
 {
@@ -421,7 +440,7 @@ void SDataprepConsumerWidget::OnConsumerChanged()
 }
 
 
-TSharedRef< SWidget > SDataprepDetailsView::CreateDefaultWidget( TSharedPtr< SWidget >& NameWidget, TSharedPtr< SWidget >& ValueWidget, float LeftPadding, EHorizontalAlignment HAlign, EVerticalAlignment VAlign, const FDataprepParameterizationContext& ParameterizationContext)
+void SDataprepDetailsView::CreateDefaultWidget( int32 Index, TSharedPtr< SWidget >& NameWidget, TSharedPtr< SWidget >& ValueWidget, float LeftPadding, EHorizontalAlignment HAlign, EVerticalAlignment VAlign, const FDataprepParameterizationContext& ParameterizationContext)
 {
 	TSharedRef<SHorizontalBox> NameColumn = SNew(SHorizontalBox)
 		.Clipping(EWidgetClipping::OnDemand);
@@ -430,23 +449,23 @@ TSharedRef< SWidget > SDataprepDetailsView::CreateDefaultWidget( TSharedPtr< SWi
 
 	// Add the name widget
 	NameColumn->AddSlot()
-		.VAlign(VAlign_Center)
-		.HAlign(HAlign_Left)
-		.Padding(FMargin(LeftPadding, 0.f, 0.f, 0.f))
-		[
-			NameWidget.ToSharedRef()
-		];
+	.VAlign(VAlign_Center)
+	.HAlign(HAlign_Left)
+	.Padding(FMargin(LeftPadding, 0.f, 0.f, 0.f))
+	[
+		NameWidget.ToSharedRef()
+	];
 
 	if ( ParameterizationContext.State == EParametrizationState::IsParameterized )
 	{
 		NameColumn->AddSlot()
-			.HAlign(HAlign_Right)
-			.VAlign(VAlign_Center)
-			.Padding(FMargin(5.f, 0.f, 5.f, 0.f))
-			.AutoWidth()
-			[
-				SNew(SDataprepParameterizationLinkIcon, DataprepAssetForParameterization.Get(), DetailedObjectAsParameterizable, ParameterizationContext.PropertyChain)
-			];
+		.HAlign(HAlign_Right)
+		.VAlign(VAlign_Center)
+		.Padding(FMargin(5.f, 0.f, 5.f, 0.f))
+		.AutoWidth()
+		[
+			SNew(SDataprepParameterizationLinkIcon, DataprepAssetForParameterization.Get(), DetailedObjectAsParameterizable, ParameterizationContext.PropertyChain)
+		];
 	}
 
 	FOnContextMenuOpening OnContextMenuOpening;
@@ -468,47 +487,66 @@ TSharedRef< SWidget > SDataprepDetailsView::CreateDefaultWidget( TSharedPtr< SWi
 			});
 	}
 
-	return SNew(SDataprepContextMenuOverride)
+	GridPanel->AddSlot(0, Index)
+	[
+		SNew(SDataprepContextMenuOverride)
 		.OnContextMenuOpening(OnContextMenuOpening)
+		[
+			DataprepWidgetUtils::CreatePropertyWidget( NameColumn, ValueWidget, ColumnSizeData, Spacing, HAlign, VAlign )
+		]
+	];
+
+	if(bColumnPadding)
+	{
+		// Add two more columns to align parameter widget
+		GridPanel->AddSlot(1, Index)
+		.Padding(5.0f, 5.0f, 0.0f, 5.0f)
 		[
 			SNew(SHorizontalBox)
 			+ SHorizontalBox::Slot()
-			.FillWidth(1.f)
 			.VAlign(VAlign_Center)
-			.HAlign(HAlign_Left)
+			.HAlign(HAlign_Right)
+			.AutoWidth()
 			[
-				SNew(SSplitter)
-				.Style(FEditorStyle::Get(), "DetailsView.Splitter")
-				.PhysicalSplitterHandleSize(1.0f)
-				.HitDetectionSplitterHandleSize(5.0f)
-				.ResizeMode( ESplitterResizeMode::Fill )
-				+ SSplitter::Slot()
-				.Value(ColumnSizeData->LeftColumnWidth)
-				.OnSlotResized( SSplitter::FOnSlotResized::CreateSP( this, &SDataprepDetailsView::OnLeftColumnResized ) )
+				SNew(SButton)
+				.IsFocusable(false)
+				.Visibility(EVisibility::Hidden)
+				.IsEnabled(false)
+				.VAlign(VAlign_Top)
+				.Content()
 				[
-					NameColumn
+					SNew(STextBlock)
+					.Font(FDataprepEditorUtils::GetGlyphFont())
+					.ColorAndOpacity(FLinearColor::Transparent)
+					.Text(FEditorFontGlyphs::Exclamation_Triangle)
 				]
-				+ SSplitter::Slot()
-				.Value(ColumnSizeData->RightColumnWidth)
-				.OnSlotResized(ColumnSizeData->OnWidthChanged)
+			]
+		];
+
+		GridPanel->AddSlot(2, Index)
+		.Padding(5.0f, 5.0f, 0.0f, 5.0f)
+		[
+			SNew(SHorizontalBox)
+			+ SHorizontalBox::Slot()
+			.VAlign(VAlign_Center)
+			.HAlign(HAlign_Right)
+			.AutoWidth()
+			[
+				SNew(SButton)
+				.IsFocusable(false)
+				.Visibility(EVisibility::Hidden)
+				.IsEnabled(false)
+				.VAlign(VAlign_Top)
+				.Content()
 				[
-					SNew(SHorizontalBox)
-					.Clipping(EWidgetClipping::OnDemand)
-					+ SHorizontalBox::Slot()
-					.Padding( 5.0f, 2.5f, 2.0f, 2.5f )
-					.HAlign( HAlign )
-					.VAlign( VAlign )
-					[
-						// Trick to force the splitter widget to fill up the space of its parent
-						// Strongly inspired from SDetailSingleItemRow
-						SNew(SConstrainedBox)
-						[
-							ValueWidget.ToSharedRef()
-						]
-					]
+					SNew(STextBlock)
+					.Font(FDataprepEditorUtils::GetGlyphFont())
+					.ColorAndOpacity(FLinearColor::Transparent)
+					.Text(FEditorFontGlyphs::Exclamation_Triangle)
 				]
 			]
 		];
+	}
 }
 
 void SDataprepDetailsView::OnPropertyChanged(const FPropertyChangedEvent& InEvent)
@@ -548,7 +586,7 @@ void SDataprepDetailsView::OnDataprepParameterizationStatusForObjectsChanged(con
 	}
 }
 
-void SDataprepDetailsView::AddWidgets( const TArray< TSharedRef< IDetailTreeNode > >& DetailTree, TSharedPtr<SGridPanel>& GridPanel, int32& Index, float LeftPadding, const FDataprepParameterizationContext& InParameterizationContext)
+void SDataprepDetailsView::AddWidgets( const TArray< TSharedRef< IDetailTreeNode > >& DetailTree, int32& Index, float LeftPadding, const FDataprepParameterizationContext& InParameterizationContext)
 {
 	auto IsDetailNodeDisplayable = []( const TSharedPtr< IPropertyHandle >& PropertyHandle)
 	{
@@ -629,7 +667,7 @@ void SDataprepDetailsView::AddWidgets( const TArray< TSharedRef< IDetailTreeNode
 
 			TArray< TSharedRef< IDetailTreeNode > > Children;
 			ChildNode->GetChildren( Children );
-			AddWidgets( Children, GridPanel, Index, LeftPadding, CurrentParameterizationContext );
+			AddWidgets( Children, Index, LeftPadding, CurrentParameterizationContext );
 		}
 		else if( IsDetailNodeDisplayableContainerProperty( PropertyHandle ) )
 		{
@@ -640,11 +678,8 @@ void SDataprepDetailsView::AddWidgets( const TArray< TSharedRef< IDetailTreeNode
 				TSharedPtr< SWidget > NameWidget;
 				TSharedPtr< SWidget > ValueWidget;
 				DetailPropertyRow->GetDefaultWidgets( NameWidget, ValueWidget, Row, true );
-				GridPanel->AddSlot(0, Index)
-				[
-					CreateDefaultWidget( NameWidget, ValueWidget, LeftPadding, Row.ValueWidget.HorizontalAlignment, Row.ValueWidget.VerticalAlignment, CurrentParameterizationContext )
-				];
 
+				CreateDefaultWidget( Index, NameWidget, ValueWidget, LeftPadding, Row.ValueWidget.HorizontalAlignment, Row.ValueWidget.VerticalAlignment, CurrentParameterizationContext );
 				Index++;
 
 				TArray< TSharedRef< IDetailTreeNode > > Children;
@@ -652,7 +687,7 @@ void SDataprepDetailsView::AddWidgets( const TArray< TSharedRef< IDetailTreeNode
 				if( Children.Num() > 0 )
 				{
 					// #ueent_todo: Find a way to add collapse/expand capability for property with children
-					AddWidgets( Children, GridPanel, Index, LeftPadding + 10.f, CurrentParameterizationContext );
+					AddWidgets( Children, Index, LeftPadding + 10.f, CurrentParameterizationContext );
 				}
 
 				TrackedProperties.Add( PropertyHandle->GetProperty() );
@@ -685,10 +720,7 @@ void SDataprepDetailsView::AddWidgets( const TArray< TSharedRef< IDetailTreeNode
 
 			if( NameWidget.IsValid() && ValueWidget.IsValid() )
 			{
-				GridPanel->AddSlot(0, Index)
-				[
-					CreateDefaultWidget( NameWidget, ValueWidget, LeftPadding, HAlign, VAlign, CurrentParameterizationContext )
-				];
+				CreateDefaultWidget( Index, NameWidget, ValueWidget, LeftPadding, HAlign, VAlign, CurrentParameterizationContext );
 				Index++;
 
 				bool bDisplayChildren = true;
@@ -708,7 +740,7 @@ void SDataprepDetailsView::AddWidgets( const TArray< TSharedRef< IDetailTreeNode
 				if( bDisplayChildren && Children.Num() > 0 )
 				{
 					// #ueent_todo: Find a way to add collapse/expand capability for property with children
-					AddWidgets( Children, GridPanel, Index, LeftPadding + 10.f, CurrentParameterizationContext );
+					AddWidgets( Children, Index, LeftPadding + 10.f, CurrentParameterizationContext );
 				}
 			}
 		}
@@ -719,6 +751,8 @@ void SDataprepDetailsView::Construct(const FArguments& InArgs)
 {
 	bRefreshObjectToDisplay = false;
 	DetailedObject = InArgs._Object;
+	Spacing = InArgs._Spacing;
+	bColumnPadding = InArgs._ColumnPadding;
 
 	if (InArgs._ColumnSizeData.IsValid())
 	{
@@ -776,11 +810,11 @@ void SDataprepDetailsView::Construct()
 		ParameterizationContext.State = DataprepAsset && DetailedObjectAsParameterizable ? EParametrizationState::CanBeParameterized : EParametrizationState::InvalidForParameterization;
 		DataprepAssetForParameterization = DataprepAsset;
 
-		TSharedPtr<SGridPanel> GridPanel = SNew(SGridPanel).FillColumn( 0.0f, 1.0f );
+		GridPanel = SNew(SGridPanel).FillColumn( 0.0f, 1.0f );
 		TArray< TSharedRef< IDetailTreeNode > > RootNodes = Generator->GetRootTreeNodes();
 
 		int32 Index = 0;
-		AddWidgets(RootNodes, GridPanel, Index, 0.f, ParameterizationContext);
+		AddWidgets(RootNodes, Index, 0.f, ParameterizationContext);
 
 		ChildSlot
 		[
@@ -904,4 +938,172 @@ FReply SDataprepContextMenuOverride::OnMouseButtonUp(const FGeometry& MyGeometry
 	return FReply::Unhandled();
 }
 
+void SDataprepInstanceParentWidget::Construct(const FArguments& InArgs)
+{
+	DataprepInstancePtr = InArgs._DataprepInstance;
+	if(!DataprepInstancePtr.IsValid())
+	{
+		return;
+	}
+
+	if (InArgs._ColumnSizeData.IsValid())
+	{
+		ColumnSizeData = InArgs._ColumnSizeData;
+	}
+	else
+	{
+		ColumnWidth = 0.7f;
+		ColumnSizeData = MakeShared<FDataprepDetailsViewColumnSizeData>();
+		ColumnSizeData->LeftColumnWidth = TAttribute<float>(this, &SDataprepInstanceParentWidget::OnGetLeftColumnWidth);
+		ColumnSizeData->RightColumnWidth = TAttribute<float>(this, &SDataprepInstanceParentWidget::OnGetRightColumnWidth);
+		ColumnSizeData->OnWidthChanged = SSplitter::FOnSlotResized::CreateSP(this, &SDataprepInstanceParentWidget::OnSetColumnWidth);
+	}
+
+
+	TSharedRef<SWidget> NameWidget = SNew(SHorizontalBox)
+	.Clipping(EWidgetClipping::OnDemand)
+	+ SHorizontalBox::Slot()
+	.VAlign(VAlign_Center)
+	.HAlign(HAlign_Left)
+	.Padding(FMargin(0.f, 0.f, 0.f, 0.f))
+	[
+		SNew( STextBlock )
+		.Text( LOCTEXT("DataprepInstanceParentWidget_Parent_Label", "Parent") )
+		.Font( IDetailLayoutBuilder::GetDetailFont() )
+	];
+
+	TSharedRef<SWidget> ValueWidget = SNew( SObjectPropertyEntryBox )
+	.AllowedClass( UDataprepAssetInterface::StaticClass() )
+	.OnObjectChanged( this, &SDataprepInstanceParentWidget::SetDataprepInstanceParent )
+	.OnShouldFilterAsset( this, &SDataprepInstanceParentWidget::ShouldFilterAsset )
+	.ObjectPath( this, &SDataprepInstanceParentWidget::GetDataprepInstanceParent );
+
+	// The widget is disabled as the workflow to change the parent asset of a Dataprep instance is refined
+	ValueWidget->SetEnabled( false );
+
+	ChildSlot
+	[
+		SNew(SHorizontalBox)
+		+SHorizontalBox::Slot()
+		.FillWidth(1)
+		[
+			SNew(SVerticalBox)
+			+ SVerticalBox::Slot()
+			.AutoHeight()
+			[
+				SNew(SHorizontalBox)
+				+ SHorizontalBox::Slot()
+				[
+					SNew(SSpacer)
+					.Size( FVector2D( 200, 10 ) )		
+				]
+			]
+			+ SVerticalBox::Slot()
+			.AutoHeight()
+			[
+				DataprepWidgetUtils::CreatePropertyWidget( NameWidget, ValueWidget, ColumnSizeData, 0.0f )
+			]
+		]
+	];
+}
+
+void SDataprepInstanceParentWidget::SetDataprepInstanceParent(const FAssetData& InAssetData)
+{
+	if(UDataprepAssetInstance* DataprepInstance = DataprepInstancePtr.Get())
+	{
+		if ( UDataprepAsset* DataprepAsset = Cast<UDataprepAsset>( InAssetData.GetAsset() ) )
+		{
+			const FScopedTransaction Transaction( LOCTEXT("DataprepInstance_SetParent", "Set Parent") );
+			DataprepInstance->SetParent( DataprepAsset );
+		}
+	}
+}
+
+FString SDataprepInstanceParentWidget::GetDataprepInstanceParent() const
+{
+	FString PathName;
+
+	if(UDataprepAssetInstance* DataprepInstance = DataprepInstancePtr.Get())
+	{
+		if(DataprepInstance->GetParent())
+		{
+			PathName = DataprepInstance->GetParent()->GetPathName();
+		}
+	}
+
+	return PathName;
+}
+
+bool SDataprepInstanceParentWidget::ShouldFilterAsset(const FAssetData& InAssetData)
+{
+	if(UDataprepAssetInstance* DataprepInstance = DataprepInstancePtr.Get())
+	{
+		if ( InAssetData.GetClass() == UDataprepAssetInterface::StaticClass() )
+		{
+			FAssetData CurrentAssetData(DataprepInstance->GetParent());
+			return CurrentAssetData != InAssetData;
+		}
+	}
+
+	return false;
+}
+
+TSharedRef<SWidget> DataprepWidgetUtils::CreateParameterRow( TSharedPtr<SWidget> ParameterWidget )
+{
+	return  SNew(SGridPanel)
+	.FillColumn(0, 1.0f)
+	+ SGridPanel::Slot(0, 0)
+	.Padding(10.0f, 5.0f, 0.0f, 5.0f)
+	[
+		ParameterWidget.ToSharedRef()
+	]
+	// Add two more columns to align parameter widget
+	+ SGridPanel::Slot(1, 0)
+	.Padding(5.0f, 5.0f, 0.0f, 5.0f)
+	[
+		SNew(SHorizontalBox)
+		+ SHorizontalBox::Slot()
+		.VAlign(VAlign_Center)
+		.HAlign(HAlign_Right)
+		.AutoWidth()
+		[
+			SNew(SButton)
+			.IsFocusable(false)
+			.Visibility(EVisibility::Hidden)
+			.IsEnabled(false)
+			.VAlign(VAlign_Top)
+			.Content()
+			[
+				SNew(STextBlock)
+				.Font(FDataprepEditorUtils::GetGlyphFont())
+				.ColorAndOpacity(FLinearColor::Transparent)
+				.Text(FEditorFontGlyphs::Exclamation_Triangle)
+			]
+		]
+	]
+	+ SGridPanel::Slot(2, 0)
+	.Padding(5.0f, 5.0f, 0.0f, 5.0f)
+	[
+		SNew(SHorizontalBox)
+		+ SHorizontalBox::Slot()
+		.VAlign(VAlign_Center)
+		.HAlign(HAlign_Right)
+		.AutoWidth()
+		[
+			SNew(SButton)
+			.IsFocusable(false)
+			.Visibility(EVisibility::Hidden)
+			.IsEnabled(false)
+			.VAlign(VAlign_Top)
+			.Content()
+			[
+				SNew(STextBlock)
+				.Font(FDataprepEditorUtils::GetGlyphFont())
+				.ColorAndOpacity(FLinearColor::Transparent)
+				.Text(FEditorFontGlyphs::Exclamation_Triangle)
+			]
+		]
+	];
+}
+
 #undef LOCTEXT_NAMESPACE
diff --git a/Engine/Plugins/Experimental/DataPrepEditor/Source/DataPrepEditor/Private/Widgets/DataprepWidgets.h b/Engine/Plugins/Experimental/DataPrepEditor/Source/DataPrepEditor/Private/Widgets/DataprepWidgets.h
--- a/Engine/Plugins/Experimental/DataPrepEditor/Source/DataPrepEditor/Private/Widgets/DataprepWidgets.h
+++ b/Engine/Plugins/Experimental/DataPrepEditor/Source/DataPrepEditor/Private/Widgets/DataprepWidgets.h
@@ -12,10 +12,44 @@
 #include "Widgets/SCompoundWidget.h"
 
 class SEditableTextBox;
+class SGridPanel;
 class UDataprepAsset;
+class UDataprepAssetInstance;
 class UDataprepContentConsumer;
 class UDataprepParameterizableObject;
 
+struct FAssetData;
+
+namespace DataprepWidgetUtils
+{
+	TSharedRef<SWidget> CreateParameterRow( TSharedPtr<SWidget> ParameterWidget );
+}
+
+/** Helper class to force a widget to fill in a space. Copied from SDetailSingleItemRow.cpp */
+class SConstrainedBox : public SCompoundWidget
+{
+public:
+	SLATE_BEGIN_ARGS(SConstrainedBox) {}
+		SLATE_DEFAULT_SLOT(FArguments, Content)
+	SLATE_END_ARGS()
+
+	void Construct(const FArguments& InArgs)
+	{
+		ChildSlot
+		[
+			InArgs._Content.Widget
+		];
+	}
+
+	virtual FVector2D ComputeDesiredSize(float LayoutScaleMultiplier) const override
+	{
+		// Voluntarily ridiculously large value to force the child widget to fill up the available space
+		const float MinWidthVal = 2000;
+		const FVector2D ChildSize = ChildSlot.GetWidget()->GetDesiredSize();
+		return FVector2D(FMath::Max(MinWidthVal, ChildSize.X), ChildSize.Y);
+	}
+};
+
 struct FDataprepDetailsViewColumnSizeData
 {
 	TAttribute<float> LeftColumnWidth;
@@ -90,9 +124,13 @@ class SDataprepDetailsView : public SCompoundWidget, public FGCObject
 public:
 	SLATE_BEGIN_ARGS(SDataprepDetailsView)
 		: _Object( nullptr )
-	{}
-	SLATE_ARGUMENT(TSharedPtr<FDataprepDetailsViewColumnSizeData>, ColumnSizeData)
-	SLATE_ARGUMENT(UObject*, Object)
+		, _Spacing(0)
+		, _ColumnPadding(false)
+		{}
+		SLATE_ARGUMENT(TSharedPtr<FDataprepDetailsViewColumnSizeData>, ColumnSizeData)
+		SLATE_ARGUMENT(UObject*, Object)
+		SLATE_ARGUMENT(float, Spacing)
+		SLATE_ARGUMENT(bool, ColumnPadding)
 	SLATE_END_ARGS()
 
 public:
@@ -125,17 +163,17 @@ private:
 	void Construct();
 
 	/** Add widgets held by array of DetailTreeNode objects */
-	void AddWidgets( const TArray< TSharedRef< class IDetailTreeNode > >& DetailTree, TSharedPtr< class SGridPanel >& GridPanel, int32& Index, float LeftPadding, const FDataprepParameterizationContext& ParameterizationContext);
+	void AddWidgets( const TArray< TSharedRef< class IDetailTreeNode > >& DetailTree, int32& Index, float LeftPadding, const FDataprepParameterizationContext& ParameterizationContext);
 
 	/**
-	 * Create generic widget for a property row
-	 * @param NameWidget The widget used in the name column
-	 * @param ValueWidget The widget used in the value column
-	 * @param LeftPadding The Padding on the left of the row
-	 * @param PropertyChain The chain to the property from the class of the detailed object
-	 * @return A row for the details view
+	 * Inserts a generic widget for a property row into the grid panel
+	 * @param Index						Row index in the grid panel
+	 * @param NameWidget				The widget used in the name column
+	 * @param ValueWidget				The widget used in the value column
+	 * @param LeftPadding				The Padding on the left of the row
+	 * @param ParameterizationContext	Parameterization context of the associated property
 	 */
-	TSharedRef< SWidget > CreateDefaultWidget(TSharedPtr< SWidget >& NameWidget, TSharedPtr< SWidget >& ValueWidget, float LeftPadding, EHorizontalAlignment HAlign, EVerticalAlignment VAlign, const FDataprepParameterizationContext& ParameterizationContext);
+	void CreateDefaultWidget(int32 Index, TSharedPtr< SWidget >& NameWidget, TSharedPtr< SWidget >& ValueWidget, float LeftPadding, EHorizontalAlignment HAlign, EVerticalAlignment VAlign, const FDataprepParameterizationContext& ParameterizationContext);
 
 	/** Callback used by all splitters in the details view, so that they move in sync */
 	void OnLeftColumnResized(float InNewWidth)
@@ -174,22 +212,38 @@ private:
 	/** Delegate handle to track property changes on array properties */
 	FDelegateHandle OnPropertyChangedHandle;
 
-	/** Delegate to track new object after a reinstancing process */
+	/** Delegate to track new object after a re-instancing process */
 	FDelegateHandle OnObjectReplacedHandle;
 
 	/** Relative width to control splitters. */
 	float ColumnWidth;
 
-	// Points to the currently used colum size data. Can be provided via argument as well.
+	/** Points to the currently used column size data. Can be provided via argument as well. */
 	TSharedPtr< FDataprepDetailsViewColumnSizeData > ColumnSizeData;
 
-	// If there is a new object to display on the next tick
+	/** If there is a new object to display on the next tick */
 	uint8 bRefreshObjectToDisplay : 1;
 
-	// This pointer to a dataprep asset is used by the parameterization system. It should be nullptr when the parameterization shouldn't be shown
+	/**
+	 * Pointer to a Dataprep asset used by the parameterization system.
+	 * It should be nullptr when the parameterization shouldn't be shown
+	 */
 	TWeakObjectPtr<UDataprepAsset> DataprepAssetForParameterization;
 
+	/** Callback used when parameterization has changed */
 	FDelegateHandle OnDataprepParameterizationStatusForObjectsChangedHandle;
+
+	/** Spacing between rows of the SDataprepDetailsView. 0 by default. */
+	float Spacing;
+
+	/**
+	 * Indicates if two columns should be added
+	 * This is useful when the SDataprepDetailsView widget is used along the Producers widget
+	 */
+	bool bColumnPadding;
+
+	/** Grid panel storing the row widgets */
+	TSharedPtr<SGridPanel> GridPanel;
 };
 
 /**
@@ -210,3 +264,41 @@ public:
 private:
 	FOnContextMenuOpening OnContextMenuOpening;
 };
+
+// Widget to expose parent asset of a Dataprep instance
+class SDataprepInstanceParentWidget : public SCompoundWidget
+{
+public:
+	SLATE_BEGIN_ARGS(SDataprepInstanceParentWidget) {}
+		SLATE_ARGUMENT(TSharedPtr< FDataprepDetailsViewColumnSizeData >, ColumnSizeData)
+		SLATE_ARGUMENT(UDataprepAssetInstance*, DataprepInstance)
+	SLATE_END_ARGS()
+
+	void Construct(const FArguments& InArgs);
+
+private:
+	void SetDataprepInstanceParent(const FAssetData& InAssetData);
+	FString GetDataprepInstanceParent() const;
+	bool ShouldFilterAsset(const FAssetData& InAssetData);
+
+	/** Callbacks to update splitter */
+	void OnLeftColumnResized(float InNewWidth)
+	{
+		// This has to be bound or the splitter will take it upon itself to determine the size
+		// We do nothing here because it is handled by the column size data
+	}
+
+	float OnGetLeftColumnWidth() const { return 1.0f - ColumnWidth; }
+	float OnGetRightColumnWidth() const { return ColumnWidth; }
+	void OnSetColumnWidth(float InWidth) { ColumnWidth = InWidth; }
+
+private:
+	/** Weak pointer on DataprepAsset instance */
+	TWeakObjectPtr<UDataprepAssetInstance> DataprepInstancePtr;
+
+	/** Helps sync column resizing with another part of the UI */
+	TSharedPtr< FDataprepDetailsViewColumnSizeData > ColumnSizeData;
+
+	/** Relative width to control splitters. */
+	float ColumnWidth;
+};
diff --git a/Engine/Plugins/Experimental/DataPrepEditor/Source/DataPrepEditor/Private/Widgets/SDataprepProducersWidget.cpp b/Engine/Plugins/Experimental/DataPrepEditor/Source/DataPrepEditor/Private/Widgets/SDataprepProducersWidget.cpp
--- a/Engine/Plugins/Experimental/DataPrepEditor/Source/DataPrepEditor/Private/Widgets/SDataprepProducersWidget.cpp
+++ b/Engine/Plugins/Experimental/DataPrepEditor/Source/DataPrepEditor/Private/Widgets/SDataprepProducersWidget.cpp
@@ -307,7 +307,7 @@ void SDataprepProducersWidget::Construct(const FArguments & InArgs, UDataprepAss
 					.VAlign(VAlign_Center)
 					[
 						SNew(STextBlock)
-						.Font( IDetailLayoutBuilder::GetDetailFontBold() )
+						.Font( FDataprepEditorStyle::Get()->GetFontStyle("DataprepEditor.SectionFont") )
 						.Text(LOCTEXT("DataprepProducersWidget_Producers_label", "Inputs"))
 					]
 					+ SHorizontalBox::Slot()
