Display a confirmation dialog when closing the dataprep editor if there is any data loaded #jira UEENT-3010 #rb JeanLuc.Corenthin

diff --git a/Engine/Plugins/Experimental/DataPrepEditor/Source/DataPrepEditor/Private/DataPrepEditor.cpp b/Engine/Plugins/Experimental/DataPrepEditor/Source/DataPrepEditor/Private/DataPrepEditor.cpp
--- a/Engine/Plugins/Experimental/DataPrepEditor/Source/DataPrepEditor/Private/DataPrepEditor.cpp
+++ b/Engine/Plugins/Experimental/DataPrepEditor/Source/DataPrepEditor/Private/DataPrepEditor.cpp
@@ -1142,6 +1142,15 @@ void FDataprepEditor::UpdatePreviewPanels(bool bInclude3DViewport)
 
 bool FDataprepEditor::OnRequestClose()
 {
+	const int ActorCount = PreviewWorld->GetActorCount();
+	if( bWorldBuilt && !bIgnoreCloseRequest && ActorCount > DefaultActorsInPreviewWorld.Num() )
+	{
+		// World was imported and is not empty -- show warning message
+		const FText Title( LOCTEXT( "DataprepEditor_ProceedWithClose", "Proceed with close") );
+		const FText Message( LOCTEXT( "DataprepEditor_ConfirmClose", "Imported data was not committed! Closing the editor will discard imported data.\nDo you want to close anyway?" ) );
+
+		return ( OpenMsgDlgInt( EAppMsgType::YesNo, Message, Title ) == EAppReturnType::Yes );
+	}
 	return !bIgnoreCloseRequest;
 }
 
