Fixed linkage error

diff --git a/Engine/Plugins/Experimental/DataPrepEditor/Source/DataPrepEditor/Private/Widgets/DataprepAssetView.cpp b/Engine/Plugins/Experimental/DataPrepEditor/Source/DataPrepEditor/Private/Widgets/DataprepAssetView.cpp
--- a/Engine/Plugins/Experimental/DataPrepEditor/Source/DataPrepEditor/Private/Widgets/DataprepAssetView.cpp
+++ b/Engine/Plugins/Experimental/DataPrepEditor/Source/DataPrepEditor/Private/Widgets/DataprepAssetView.cpp
@@ -586,7 +586,7 @@ void SDataprepAssetView::Construct( const FArguments& InArgs, UDataprepAssetInte
 						.VAlign(VAlign_Center)
 						.Padding(10.0f, 5.0f, 0.0f, 5.0f)
 						[
-							SNew( SConstrainedBox )
+							SNew( DataprepWidgetUtils::SConstrainedBox )
 							[
 								SNew( SDataprepDetailsView )
 								.Object( DataprepAssetInterfacePtr->GetParameterizationObject() )
diff --git a/Engine/Plugins/Experimental/DataPrepEditor/Source/DataPrepEditor/Private/Widgets/DataprepWidgets.cpp b/Engine/Plugins/Experimental/DataPrepEditor/Source/DataPrepEditor/Private/Widgets/DataprepWidgets.cpp
--- a/Engine/Plugins/Experimental/DataPrepEditor/Source/DataPrepEditor/Private/Widgets/DataprepWidgets.cpp
+++ b/Engine/Plugins/Experimental/DataPrepEditor/Source/DataPrepEditor/Private/Widgets/DataprepWidgets.cpp
@@ -87,7 +87,7 @@ namespace DataprepWidgetUtils
 				[
 					// Trick to force the splitter widget to fill up the space of its parent
 					// Strongly inspired from SDetailSingleItemRow
-					SNew(SConstrainedBox)
+					SNew(DataprepWidgetUtils::SConstrainedBox)
 					[
 						ValueWidget.ToSharedRef()
 					]
@@ -272,7 +272,7 @@ TSharedRef<SWidget> SDataprepConsumerWidget::BuildWidget()
 					[
 						// Trick to force the splitter widget to fill up the space of its parent
 						// Strongly inspired from SDetailSingleItemRow
-						SNew(SConstrainedBox)
+						SNew(DataprepWidgetUtils::SConstrainedBox)
 						[
 							SAssignNew(ContentFolderTextBox, SEditableTextBox)
 							.Font(IDetailLayoutBuilder::GetDetailFont())
@@ -361,7 +361,7 @@ TSharedRef<SWidget> SDataprepConsumerWidget::BuildWidget()
 					[
 						// Trick to force the splitter widget to fill up the space of its parent
 						// Strongly inspired from SDetailSingleItemRow
-						SNew(SConstrainedBox)
+						SNew(DataprepWidgetUtils::SConstrainedBox)
 						[
 							SAssignNew(LevelTextBox, SEditableTextBox)
 							.Font(IDetailLayoutBuilder::GetDetailFont())
diff --git a/Engine/Plugins/Experimental/DataPrepEditor/Source/DataPrepEditor/Private/Widgets/DataprepWidgets.h b/Engine/Plugins/Experimental/DataPrepEditor/Source/DataPrepEditor/Private/Widgets/DataprepWidgets.h
--- a/Engine/Plugins/Experimental/DataPrepEditor/Source/DataPrepEditor/Private/Widgets/DataprepWidgets.h
+++ b/Engine/Plugins/Experimental/DataPrepEditor/Source/DataPrepEditor/Private/Widgets/DataprepWidgets.h
@@ -23,32 +23,32 @@ struct FAssetData;
 namespace DataprepWidgetUtils
 {
 	TSharedRef<SWidget> CreateParameterRow( TSharedPtr<SWidget> ParameterWidget );
-}
-
-/** Helper class to force a widget to fill in a space. Copied from SDetailSingleItemRow.cpp */
-class SConstrainedBox : public SCompoundWidget
-{
-public:
-	SLATE_BEGIN_ARGS(SConstrainedBox) {}
-		SLATE_DEFAULT_SLOT(FArguments, Content)
-	SLATE_END_ARGS()
-
-	void Construct(const FArguments& InArgs)
-	{
-		ChildSlot
-		[
-			InArgs._Content.Widget
-		];
-	}
 
-	virtual FVector2D ComputeDesiredSize(float LayoutScaleMultiplier) const override
+	/** Helper class to force a widget to fill in a space. Copied from SDetailSingleItemRow.cpp */
+	class SConstrainedBox : public SCompoundWidget
 	{
-		// Voluntarily ridiculously large value to force the child widget to fill up the available space
-		const float MinWidthVal = 2000;
-		const FVector2D ChildSize = ChildSlot.GetWidget()->GetDesiredSize();
-		return FVector2D(FMath::Max(MinWidthVal, ChildSize.X), ChildSize.Y);
-	}
-};
+	public:
+		SLATE_BEGIN_ARGS(SConstrainedBox) {}
+		SLATE_DEFAULT_SLOT(FArguments, Content)
+			SLATE_END_ARGS()
+
+			void Construct(const FArguments& InArgs)
+		{
+			ChildSlot
+				[
+					InArgs._Content.Widget
+				];
+		}
+
+		virtual FVector2D ComputeDesiredSize(float LayoutScaleMultiplier) const override
+		{
+			// Voluntarily ridiculously large value to force the child widget to fill up the available space
+			const float MinWidthVal = 2000;
+			const FVector2D ChildSize = ChildSlot.GetWidget()->GetDesiredSize();
+			return FVector2D(FMath::Max(MinWidthVal, ChildSize.X), ChildSize.Y);
+		}
+	};
+}
 
 struct FDataprepDetailsViewColumnSizeData
 {
