Remove #ueent_todo tags

diff --git a/Engine/Plugins/Enterprise/DatasmithImporter/Source/DatasmithImporter/Private/DatasmithFileProducer.cpp b/Engine/Plugins/Enterprise/DatasmithImporter/Source/DatasmithImporter/Private/DatasmithFileProducer.cpp
--- a/Engine/Plugins/Enterprise/DatasmithImporter/Source/DatasmithImporter/Private/DatasmithFileProducer.cpp
+++ b/Engine/Plugins/Enterprise/DatasmithImporter/Source/DatasmithImporter/Private/DatasmithFileProducer.cpp
@@ -392,8 +392,6 @@ void UDatasmithFileProducer::PreventNameCollision()
 
 				Object->Rename( nullptr, NewPackage, REN_DontCreateRedirectors | REN_NonTransactional );
 
-				// #ueent_todo: Should we prefix the unique Id of the asset with the producer's namespace?
-
 				AssetRedirectorMap.Emplace( PreviousObjectPath, Object );
 				if(bCheckPackage)
 				{
@@ -549,7 +547,6 @@ void UDatasmithFileProducer::PreventNameCollision()
 					}
 
 					// Remove reference to Datasmith scene
-					// #ueent_todo: Should the IDatasmithScene kept around?
 					SceneActor->Scene = nullptr;
 				}
 			}
@@ -836,7 +833,6 @@ bool UDatasmithDirProducer::Supersede(const UDataprepContentProducer* OtherProdu
 							   ( OtherDirProducer->FolderPath.StartsWith( FolderPath ) && bRecursive );
 
 		// Check if this producer will generate a super-set of the set of files generated by the other one
-		// #ueent_todo: This could be expensive...
 		if(bCouldSupersede)
 		{
 			TSet< FString > ThisFilesToProcess = GetSetOfFiles();
@@ -1200,7 +1196,6 @@ void FDatasmithFileProducerUtils::DeletePackagePath( const FString& PathToDelete
 	// Inspired from ContentBrowserUtils::DeleteFolders
 	// Inspired from ContentBrowserUtils::LoadAssetsIfNeeded
 
-	// #ueent_todo: Clean-up code
 	// Form a filter from the paths
 	FAssetRegistryModule& AssetRegistryModule = FModuleManager::Get().LoadModuleChecked<FAssetRegistryModule>(TEXT("AssetRegistry"));
 
@@ -1378,7 +1373,6 @@ FString FDatasmithFileProducerUtils::SelectDirectory()
 		const FString Title = LOCTEXT("DatasmithDirProducerFolderTitle", "Choose a folder").ToString();
 		const FString DefaultLocation( FEditorDirectories::Get().GetLastDirectory( ELastDirectory::GENERIC_IMPORT ) );
 
-		// #ueent_todo: What if the user selects a drive? Shouldn't we ask if this is really what he/she wants?
 		const bool bFolderSelected = DesktopPlatform->OpenDirectoryDialog(
 			ParentWindowHandle,
 			Title,
diff --git a/Engine/Plugins/Experimental/DataPrepEditor/Source/DataPrepCore/Private/DataPrepContentProducer.cpp b/Engine/Plugins/Experimental/DataPrepEditor/Source/DataPrepCore/Private/DataPrepContentProducer.cpp
--- a/Engine/Plugins/Experimental/DataPrepEditor/Source/DataPrepCore/Private/DataPrepContentProducer.cpp
+++ b/Engine/Plugins/Experimental/DataPrepEditor/Source/DataPrepCore/Private/DataPrepContentProducer.cpp
@@ -44,8 +44,6 @@ bool UDataprepContentProducer::Produce(const FDataprepProducerContext& InContext
 	// Cache number of assets to go through new assets after call to Execute
 	int32 LastAssetCount = OutAssets.Num();
 
-	// Prefix all newly created actors with the namespace of the producer
-	// #ueent_todo: find a better way to identify newly created actors
 	if(!Execute( OutAssets ))
 	{
 		Terminate();
@@ -72,7 +70,7 @@ bool UDataprepContentProducer::Produce(const FDataprepProducerContext& InContext
 	// Map between old path and new path of newly created actors
 	TMap< FSoftObjectPath, FSoftObjectPath > ActorRedirectorMap;
 
-	// Rename actors
+	// Prefix all newly created actors with the namespace of the producer
 	const FString Namespace = GetNamespace();
 
 	for (TActorIterator<AActor> It(Context.WorldPtr.Get(), AActor::StaticClass(), Flags); It; ++It)
diff --git a/Engine/Plugins/Experimental/DataPrepEditor/Source/DataPrepCore/Private/DataPrepFactories.cpp b/Engine/Plugins/Experimental/DataPrepEditor/Source/DataPrepCore/Private/DataPrepFactories.cpp
--- a/Engine/Plugins/Experimental/DataPrepEditor/Source/DataPrepCore/Private/DataPrepFactories.cpp
+++ b/Engine/Plugins/Experimental/DataPrepEditor/Source/DataPrepCore/Private/DataPrepFactories.cpp
@@ -40,7 +40,6 @@ UObject * UDataprepAssetFactory::FactoryCreateNew(UClass* InClass, UObject* InPa
 
 	if(ConsumerClasses.Num() == 0)
 	{
-		// #ueent_todo: Log error about missing consumer classes
 		return nullptr;
 	}
 
diff --git a/Engine/Plugins/Experimental/DataPrepEditor/Source/DataPrepCore/Private/DataPrepOperation.cpp b/Engine/Plugins/Experimental/DataPrepEditor/Source/DataPrepCore/Private/DataPrepOperation.cpp
--- a/Engine/Plugins/Experimental/DataPrepEditor/Source/DataPrepCore/Private/DataPrepOperation.cpp
+++ b/Engine/Plugins/Experimental/DataPrepEditor/Source/DataPrepCore/Private/DataPrepOperation.cpp
@@ -137,7 +137,7 @@ UObject* UDataprepEditingOperation::AddAsset( const UObject* Asset, const FStrin
 		return OperationContext->AddAssetDelegate.Execute( Asset, AssetName.Len() > 0 ? *AssetName : nullptr );
 	}
 
-	// #ueent_todo: Report error
+	UE_LOG( LogDataprepCore, Log, TEXT("UDataprepEditingOperation::AddAsset called without context") );
 
 	return nullptr;
 }
@@ -149,7 +149,7 @@ UObject* UDataprepEditingOperation::CreateAsset( UClass* AssetClass, const FStri
 		return OperationContext->CreateAssetDelegate.Execute( AssetClass, AssetName.Len() > 0 ? *AssetName : nullptr );
 	}
 
-	// #ueent_todo: Report error
+	UE_LOG( LogDataprepCore, Log, TEXT("UDataprepEditingOperation::CreateAsset called without context") );
 
 	return nullptr;
 }
@@ -161,7 +161,7 @@ AActor * UDataprepEditingOperation::CreateActor(UClass* ActorClass, const FStrin
 		return OperationContext->CreateActorDelegate.Execute( ActorClass, ActorName.Len() > 0 ? *ActorName : nullptr );
 	}
 
-	// #ueent_todo: Report error
+	UE_LOG( LogDataprepCore, Log, TEXT("UDataprepEditingOperation::CreateActor called without context") );
 
 	return nullptr;
 }
@@ -174,7 +174,7 @@ void UDataprepEditingOperation::RemoveObject( UObject* Object, bool bLocalContex
 		return;
 	}
 
-	// #ueent_todo: Report error
+	UE_LOG( LogDataprepCore, Log, TEXT("UDataprepEditingOperation::RemoveObject called without context") );
 }
 
 void UDataprepEditingOperation::RemoveObjects(TArray<UObject*> Objects, bool bLocalContext)
@@ -188,7 +188,7 @@ void UDataprepEditingOperation::RemoveObjects(TArray<UObject*> Objects, bool bLo
 		return;
 	}
 
-	// #ueent_todo: Report error
+	UE_LOG( LogDataprepCore, Log, TEXT("UDataprepEditingOperation::RemoveObjects called without context") );
 }
 
 void UDataprepEditingOperation::DeleteObject( UObject* Object )
@@ -201,7 +201,7 @@ void UDataprepEditingOperation::DeleteObject( UObject* Object )
 		return;
 	}
 
-	// #ueent_todo: Report error
+	UE_LOG( LogDataprepCore, Log, TEXT("UDataprepEditingOperation::DeleteObject called without context") );
 }
 
 void UDataprepEditingOperation::DeleteObjects( TArray<UObject*> Objects )
@@ -211,7 +211,7 @@ void UDataprepEditingOperation::DeleteObjects( TArray<UObject*> Objects )
 		OperationContext->DeleteObjectsDelegate.Execute( MoveTemp( Objects ) );
 	}
 
-	// #ueent_todo: Report error
+	UE_LOG( LogDataprepCore, Log, TEXT("UDataprepEditingOperation::DeleteObjects called without context") );
 }
 
 #undef LOCTEXT_NAMESPACE
diff --git a/Engine/Plugins/Experimental/DataPrepEditor/Source/DataPrepCore/Private/DataprepActionAsset.cpp b/Engine/Plugins/Experimental/DataPrepEditor/Source/DataPrepCore/Private/DataprepActionAsset.cpp
--- a/Engine/Plugins/Experimental/DataPrepEditor/Source/DataPrepCore/Private/DataprepActionAsset.cpp
+++ b/Engine/Plugins/Experimental/DataPrepEditor/Source/DataPrepCore/Private/DataprepActionAsset.cpp
@@ -85,7 +85,6 @@ void UDataprepActionAsset::Execute(const TArray<UObject*>& InObjects)
 			if ( UDataprepOperation* Operation = Step->Operation )
 			{
 				Operation->Execute( OperationContext->Context->Objects );
-				// #ueent_todo: Add steps to process, added, removed and deleted objects
 			}
 			else if ( UDataprepFilter* Filter = Step->Filter )
 			{
@@ -600,10 +599,6 @@ void UDataprepActionAsset::ProcessWorkingSetChanged()
 						}
 					}
 				}
-				else
-				{
-					// #ueent_todo: Warn object external to Dataprep's context has been removed
-				}
 			}
 
 			ObjectsToRemove.Empty( ObjectsToRemove.Num() );
@@ -638,10 +633,6 @@ void UDataprepActionAsset::ProcessWorkingSetChanged()
 						bWorldChanged = true;
 					}
 				}
-				else
-				{
-					// #ueent_todo: Warn object external to Dataprep's context has been deleted
-				}
 			}
 
 			FDataprepCoreUtils::PurgeObjects( MoveTemp(ObjectsToDelete) );
diff --git a/Engine/Plugins/Experimental/DataPrepEditor/Source/DataPrepCore/Private/DataprepCoreLibrary.cpp b/Engine/Plugins/Experimental/DataPrepEditor/Source/DataPrepCore/Private/DataprepCoreLibrary.cpp
--- a/Engine/Plugins/Experimental/DataPrepEditor/Source/DataPrepCore/Private/DataprepCoreLibrary.cpp
+++ b/Engine/Plugins/Experimental/DataPrepEditor/Source/DataPrepCore/Private/DataprepCoreLibrary.cpp
@@ -261,7 +261,6 @@ bool UDataprepCoreLibrary::Execute_Internal(UDataprepAssetInterface* DataprepAss
 				}
 			}
 
-			// #ueent_todo: Should we find a better way to silently delete assets?
 			// Disable warnings from LogStaticMesh because FDataprepCoreUtils::PurgeObjects is pretty verbose on harmless warnings
 			ELogVerbosity::Type PrevLogStaticMeshVerbosity = LogStaticMesh.GetVerbosity();
 			LogStaticMesh.SetVerbosity( ELogVerbosity::Error );
diff --git a/Engine/Plugins/Experimental/DataPrepEditor/Source/DataPrepEditor/Private/DataPrepEditor.cpp b/Engine/Plugins/Experimental/DataPrepEditor/Source/DataPrepEditor/Private/DataPrepEditor.cpp
--- a/Engine/Plugins/Experimental/DataPrepEditor/Source/DataPrepEditor/Private/DataPrepEditor.cpp
+++ b/Engine/Plugins/Experimental/DataPrepEditor/Source/DataPrepEditor/Private/DataPrepEditor.cpp
@@ -431,7 +431,6 @@ void FDataprepEditor::InitDataprepEditor(const EToolkitMode::Type Mode, const TS
 
 		// Necessary step to regenerate blueprint generated class
 		// Note that this compilation will always succeed as Dataprep node does not have real body
-		// #ueent_todo: Is there a better solution
 		{
 			FKismetEditorUtilities::CompileBlueprint( DataprepRecipeBPPtr.Get(), EBlueprintCompileOptions::None, nullptr );
 		}
@@ -543,7 +542,6 @@ void FDataprepEditor::OnBuildWorld()
 	{
 		FTimeLogger TimeLogger( TEXT("Import") );
 
-		// #ueent_todo: Add progress logger to Dataprep editor
 		FDataprepProducerContext Context;
 		Context.SetWorld( PreviewWorld.Get() )
 			.SetRootPackage( TransientPackage )
@@ -647,7 +645,6 @@ void FDataprepEditor::CleanPreviewWorld()
 		}
 	}
 
-	// #ueent_todo: Should we find a better way to silently delete assets?
 	// Disable warnings from LogStaticMesh because FDataprepCoreUtils::PurgeObjects is pretty verbose on harmless warnings
 	ELogVerbosity::Type PrevLogStaticMeshVerbosity = LogStaticMesh.GetVerbosity();
 	LogStaticMesh.SetVerbosity( ELogVerbosity::Error );
@@ -765,7 +762,7 @@ void FDataprepEditor::OnCommitWorld()
 
 	if( !DataprepAssetInterfacePtr->RunConsumer( Context ) )
 	{
-		// #ueent_todo: Inform consumer failed
+		UE_LOG( LogDataprepEditor, Error, TEXT("Consumer failed...") );
 	}
 
 	ResetBuildWorld();
@@ -1123,8 +1120,6 @@ void FDataprepEditor::UpdatePreviewPanels(bool bInclude3DViewport)
 	{
 		FTimeLogger TimeLogger( TEXT("UpdatePreviewPanels") );
 
-		// #ueent_todo: There should be a event triggered to inform listeners
-		//				   that new assets have been generated.
 		AssetPreviewView->ClearAssetList();
 		FString SubstitutePath = DataprepAssetInterfacePtr->GetOutermost()->GetName();
 		if(DataprepAssetInterfacePtr->GetConsumer() != nullptr && !DataprepAssetInterfacePtr->GetConsumer()->GetTargetContentFolder().IsEmpty())
@@ -1179,8 +1174,6 @@ bool FDataprepEditor::OnCanExecuteNextStep(UDataprepActionAsset* ActionAsset, UD
 {
 	// #ueent_todo: Make this action configurable by user
 	UpdatePreviewPanels(false);
-
-	// #ueent_todo: Make this action configurable by user
 	return true;
 }
 
@@ -1224,7 +1217,6 @@ void DataprepEditorUtil::DeleteTemporaryPackage( const FString& PathToDelete )
 
 		if(AssetsToDelete.Num() > 0)
 		{
-			// #ueent_todo: We should not use ObjectTools::DeleteObjects but FAssetDeleteModel 
 			ObjectTools::DeleteObjects(AssetsToDelete, false);
 		}
 	}
diff --git a/Engine/Plugins/Experimental/DataPrepEditor/Source/DataPrepEditor/Private/DataprepSnapshot.cpp b/Engine/Plugins/Experimental/DataPrepEditor/Source/DataPrepEditor/Private/DataprepSnapshot.cpp
--- a/Engine/Plugins/Experimental/DataPrepEditor/Source/DataPrepEditor/Private/DataprepSnapshot.cpp
+++ b/Engine/Plugins/Experimental/DataPrepEditor/Source/DataPrepEditor/Private/DataprepSnapshot.cpp
@@ -49,10 +49,6 @@ enum class EDataprepAssetClass : uint8 {
 	EMaxClasses
 };
 
-// #ueent_todo: Boolean driving activating actual snapshot based logic
-const bool bUseSnapshot = true;
-const bool bUseCompression = false;
-
 namespace DataprepSnapshotUtil
 {
 	const TCHAR* SnapshotExtension = TEXT(".dpc");
@@ -98,7 +94,6 @@ namespace DataprepSnapshotUtil
 
 	}
 
-	// #ueent_todo: Find a solution using the path of the RootPackage instead of the UPackage object itself
 	FString BuildAssetFileName(const FString& RootPath, const FString& AssetPath )
 	{
 		static FString FileNamePrefix( TEXT("stream_") );
@@ -170,7 +165,6 @@ namespace DataprepSnapshotUtil
 			}
 
 			// Sort array of sub-objects: first objects do not depend on ones below
-			// #ueent_todo: Improve performance of building. Current is pretty brute force
 			int32 Count = SubObjectsArray.Num();
 			SubObjectsArray.Empty(Count);
 
@@ -241,28 +235,6 @@ namespace DataprepSnapshotUtil
 			bool bRebuildResource = !!Texture->Resource;
 			Ar << bRebuildResource;
 		}
-
-		if(bUseCompression)
-		{
-			const int32 BufferHeaderSize = (int32)sizeof(int32);
-
-			TArray<uint8> MemoryBuffer( MoveTemp( OutSerializedData ) );
-
-			// allocate all of the input space for the output, with extra space for max overhead of zlib (when compressed > uncompressed)
-			int32 CompressedSize = FCompression::CompressMemoryBound(NAME_Zlib, MemoryBuffer.Num());
-			OutSerializedData.SetNum( CompressedSize + BufferHeaderSize );
-
-			// Store the size of the uncompressed buffer
-			((int32*)OutSerializedData.GetData())[0] = MemoryBuffer.Num();
-
-			// compress the data		
-			bool bSucceeded = FCompression::CompressMemory(NAME_Zlib, OutSerializedData.GetData() + BufferHeaderSize, CompressedSize, MemoryBuffer.GetData(), MemoryBuffer.Num());
-
-			// if it failed send the data uncompressed, which we mark by setting compressed size to 0
-			checkf(bSucceeded, TEXT("zlib failed to compress, which is very unexpected"));
-
-			OutSerializedData.SetNum( CompressedSize + BufferHeaderSize );
-		}
 	}
 
 	void ReadSnapshotData(UObject* Object, const TArray<uint8>& InSerializedData, TMap<FString, UClass*>& InClassesMap, TArray<UObject*>& ObjectsToDelete)
@@ -281,25 +253,9 @@ namespace DataprepSnapshotUtil
 			}
 		};
 
-		TArray<uint8> MemoryBuffer;
-		if(bUseCompression)
-		{
-			const int32 BufferHeaderSize = (int32)sizeof(int32);
-
-			// Allocate the space required for the uncompressed data
-			int32 UncompressedSize = ((int32*)InSerializedData.GetData())[0];
-			MemoryBuffer.SetNum(UncompressedSize);
-
-			// uncompress the data		
-			bool bSucceeded = FCompression::UncompressMemory(NAME_Zlib, MemoryBuffer.GetData(), MemoryBuffer.Num(), InSerializedData.GetData() + BufferHeaderSize, InSerializedData.Num() - BufferHeaderSize);
-
-			// if it failed send the data uncompressed, which we mark by setting compressed size to 0
-			checkf(bSucceeded, TEXT("zlib failed to uncompress, which is very unexpected"));
-		}
-
 		RemoveDefaultDependencies( Object );
 
-		FMemoryReader MemAr(bUseCompression ? MemoryBuffer : InSerializedData);
+		FMemoryReader MemAr(InSerializedData);
 		FSnapshotCustomArchive Ar(MemAr);
 
 		// Deserialize count of sub-objects
@@ -446,11 +402,6 @@ public:
 
 void FDataprepEditor::TakeSnapshot()
 {
-	if(!bUseSnapshot)
-	{
-		return;
-	}
-
 	TRACE_CPUPROFILER_EVENT_SCOPE(FDataprepEditor::TakeSnapshot);
 
 	uint64 StartTime = FPlatformTime::Cycles64();
@@ -644,7 +595,6 @@ void FDataprepEditor::TakeSnapshot()
 		return;
 	}
 
-	// #ueent_todo: Is that necessary since Assets has already been sorted?
 	ContentSnapshot.DataEntries.Sort([&](const FSnapshotDataEntry& A, const FSnapshotDataEntry& B)
 	{
 		return GetAssetClassEnum( A.Get<1>() ) < GetAssetClassEnum( B.Get<1>() );
@@ -660,16 +610,10 @@ void FDataprepEditor::TakeSnapshot()
 
 void FDataprepEditor::RestoreFromSnapshot(bool bUpdateViewport)
 {
-	if(!bUseSnapshot)
-	{
-		OnBuildWorld();
-		return;
-	}
-
 	// Snapshot is not usable, rebuild the world from the producers
 	if ( !ContentSnapshot.bIsValid )
 	{
-		// #ueent_todo: Inform user that snapshot is no good and world is going to be rebuilt from scratch
+		UE_LOG( LogDataprepEditor, Log, TEXT("Snapshot is invalid. Running the producers...") );
 		OnBuildWorld();
 		return;
 	}
diff --git a/Engine/Plugins/Experimental/DataPrepEditor/Source/DataPrepEditor/Private/Widgets/DataprepAssetView.cpp b/Engine/Plugins/Experimental/DataPrepEditor/Source/DataPrepEditor/Private/Widgets/DataprepAssetView.cpp
--- a/Engine/Plugins/Experimental/DataPrepEditor/Source/DataPrepEditor/Private/Widgets/DataprepAssetView.cpp
+++ b/Engine/Plugins/Experimental/DataPrepEditor/Source/DataPrepEditor/Private/Widgets/DataprepAssetView.cpp
@@ -413,7 +413,6 @@ void SDataprepAssetView::Construct( const FArguments& InArgs, UDataprepAssetInte
 				.Padding( 0, 10, 0, 0 )
 				.HAlign( EHorizontalAlignment::HAlign_Center )
 				[
-					// #ueent_todo: make color block's width vary with parent widget
 					SNew(SColorBlock)
 					.Color( FLinearColor( 0.9f, 0.9f, 0.9f ) )
 					.Size( FVector2D( 3000, 1 ) )
@@ -471,7 +470,6 @@ void SDataprepAssetView::Construct( const FArguments& InArgs, UDataprepAssetInte
 						.Padding( 0, 10, 0, 0 )
 						.HAlign( EHorizontalAlignment::HAlign_Center )
 						[
-							// #ueent_todo: make color block's width vary with parent widget
 							SNew(SColorBlock)
 							.Color( FLinearColor( 0.9f, 0.9f, 0.9f ) )
 							.Size( FVector2D( 3000, 1 ) )
diff --git a/Engine/Plugins/Experimental/DataPrepEditor/Source/DataPrepEditor/Private/Widgets/DataprepAssetView.h b/Engine/Plugins/Experimental/DataPrepEditor/Source/DataPrepEditor/Private/Widgets/DataprepAssetView.h
--- a/Engine/Plugins/Experimental/DataPrepEditor/Source/DataPrepEditor/Private/Widgets/DataprepAssetView.h
+++ b/Engine/Plugins/Experimental/DataPrepEditor/Source/DataPrepEditor/Private/Widgets/DataprepAssetView.h
@@ -60,7 +60,6 @@ private:
 };
 
 // Inspired from SKismetInspector class
-// #ueent_todo: Revisit SGraphNodeDetailsWidget class based on tpm's feedback
 class SGraphNodeDetailsWidget : public SCompoundWidget
 {
 public:
diff --git a/Engine/Plugins/Experimental/DataPrepEditor/Source/DataPrepEditor/Private/Widgets/DataprepWidgets.cpp b/Engine/Plugins/Experimental/DataPrepEditor/Source/DataPrepEditor/Private/Widgets/DataprepWidgets.cpp
--- a/Engine/Plugins/Experimental/DataPrepEditor/Source/DataPrepEditor/Private/Widgets/DataprepWidgets.cpp
+++ b/Engine/Plugins/Experimental/DataPrepEditor/Source/DataPrepEditor/Private/Widgets/DataprepWidgets.cpp
@@ -103,7 +103,7 @@ void SDataprepConsumerWidget::OnLevelNameChanged( const FText &NewLevelName, ETe
 	FText OutReason;
 	if( !DataprepConsumerPtr->SetLevelName( NewLevelName.ToString(), OutReason ) )
 	{
-		// #ueent_todo: Warn user name is wrong
+		UE_LOG( LogDataprepEditor, Error, TEXT("%s"), *OutReason.ToString() );
 		LevelTextBox->SetText( FText::FromString( DataprepConsumerPtr->GetLevelName() ) );
 	}
 }
diff --git a/Engine/Plugins/Experimental/DataPrepEditor/Source/DataPrepEditor/Private/Widgets/SDataprepEditorViewport.cpp b/Engine/Plugins/Experimental/DataPrepEditor/Source/DataPrepEditor/Private/Widgets/SDataprepEditorViewport.cpp
--- a/Engine/Plugins/Experimental/DataPrepEditor/Source/DataPrepEditor/Private/Widgets/SDataprepEditorViewport.cpp
+++ b/Engine/Plugins/Experimental/DataPrepEditor/Source/DataPrepEditor/Private/Widgets/SDataprepEditorViewport.cpp
@@ -80,7 +80,6 @@ const FColor PerMeshColor[] =
 	FColor(  0,120, 49),
 };
 
-// #ueent_todo: Remove this debug utility before releasing
 namespace ViewportDebug
 {
 	static bool bLogTiming = false;
@@ -118,7 +117,7 @@ namespace ViewportDebug
 }
 
 // Extension of the FStaticMeshSceneProxy class to allow wireframe display per individual mesh
-// #ueent_todo: Could we unify this behavior across classes deriving from UMeshComponent?
+// #ueent_remark: Could we unify this behavior across classes deriving from UMeshComponent?
 class FStaticMeshSceneProxyExt : public FStaticMeshSceneProxy
 {
 public:
@@ -186,7 +185,6 @@ namespace DataprepEditor3DPreviewUtils
 		TArray< T* > Result;
 
 		const EActorIteratorFlags Flags = EActorIteratorFlags::SkipPendingKill;
-		// #ueent_todo: Maybe just enumerate the actors in the current level
 		for (TActorIterator<AActor> It( World, AActor::StaticClass(), Flags ); It; ++It)
 		{
 			AActor* Actor = *It;
@@ -1256,11 +1254,6 @@ void FDataprepEditorViewportClient::ProcessClick(FSceneView& View, HHitProxy* Hi
 
 							return;
 						}
-						// A contextual menu is requested
-						else if(Key == EKeys::RightMouseButton)
-						{
-							// #ueent_todo: Display contextual menu
-						}
 					}
 				}
 			}
@@ -1362,7 +1355,7 @@ TSharedRef<SWidget> SDataprepEditorViewportToolbar::GenerateShowMenu() const
 		ShowMenuBuilder.AddMenuEntry(FDataprepEditorViewportCommands::Get().SetShowBounds);
 	}
 
-	// #ueent_todo: Look at SAnimViewportToolBar::GenerateShowMenu in SAnimViewportToolBar.cpp for adding ShowFlagFilter to the Show menu
+	// #ueent_remark: Look at SAnimViewportToolBar::GenerateShowMenu in SAnimViewportToolBar.cpp for adding ShowFlagFilter to the Show menu
 
 	return ShowMenuBuilder.MakeWidget();
 }
diff --git a/Engine/Plugins/Experimental/DataPrepEditor/Source/DataPrepEditor/Private/Widgets/SDataprepProducersWidget.cpp b/Engine/Plugins/Experimental/DataPrepEditor/Source/DataPrepEditor/Private/Widgets/SDataprepProducersWidget.cpp
--- a/Engine/Plugins/Experimental/DataPrepEditor/Source/DataPrepEditor/Private/Widgets/SDataprepProducersWidget.cpp
+++ b/Engine/Plugins/Experimental/DataPrepEditor/Source/DataPrepEditor/Private/Widgets/SDataprepProducersWidget.cpp
@@ -60,7 +60,6 @@ void FContentProducerEntry::ToggleProducer()
 		const FScopedTransaction Transaction( LOCTEXT("Producers_ToggleProducer", "Toggle Producer") );
 		AssetProducers->EnableProducer(ProducerIndex, !bIsEnabled);
 
-		// #ueent_todo: Cache previous value to report failed enabling/disabling
 		bIsEnabled = AssetProducers->IsProducerEnabled(ProducerIndex);
 	}
 }
diff --git a/Engine/Plugins/Experimental/DataPrepEditor/Source/DataPrepLibraries/Private/DataprepEditingOperations.cpp b/Engine/Plugins/Experimental/DataPrepEditor/Source/DataPrepLibraries/Private/DataprepEditingOperations.cpp
--- a/Engine/Plugins/Experimental/DataPrepEditor/Source/DataPrepLibraries/Private/DataprepEditingOperations.cpp
+++ b/Engine/Plugins/Experimental/DataPrepEditor/Source/DataPrepLibraries/Private/DataprepEditingOperations.cpp
@@ -115,7 +115,6 @@ void UDataprepDeleteObjectsOperation::OnExecution_Implementation(const FDataprep
 		if (AActor* Actor = Cast< AActor >( Object ))
 		{
 			ActorsToDelete.Add(FActorAndDepth{Actor, DatasmithEditingOperationsUtils::GetActorDepth(Actor)});
-			// #ueent_todo if rem children option, add them here
 		}
 		else if(FDataprepCoreUtils::IsAsset(Object))
 		{
@@ -681,6 +680,7 @@ namespace DatasmithEditingOperationsUtils
 			{
 				if(!Actor->IsPendingKillOrUnreachable())
 				{
+					// Set current world to first world encountered
 					if(World == nullptr)
 					{
 						World = Actor->GetWorld();
@@ -688,7 +688,7 @@ namespace DatasmithEditingOperationsUtils
 
 					if(World != Actor->GetWorld())
 					{
-						// #ueent_todo: Warn that incompatible actor found and discarded
+						UE_LOG( LogDataprep, Log, TEXT("Actor %s is not part of the Dataprep transient world ..."), *Actor->GetActorLabel() );
 						continue;
 					}
 
diff --git a/Engine/Plugins/Experimental/DataPrepEditor/Source/DataPrepLibraries/Private/DataprepOperations.cpp b/Engine/Plugins/Experimental/DataPrepEditor/Source/DataPrepLibraries/Private/DataprepOperations.cpp
--- a/Engine/Plugins/Experimental/DataPrepEditor/Source/DataPrepLibraries/Private/DataprepOperations.cpp
+++ b/Engine/Plugins/Experimental/DataPrepEditor/Source/DataPrepLibraries/Private/DataprepOperations.cpp
@@ -55,13 +55,18 @@ namespace DataprepOperationTime
 
 void UDataprepSetLODsOperation::OnExecution_Implementation(const FDataprepContext& InContext)
 {
+	if(ReductionSettings.Num() > MAX_STATIC_MESH_LODS)
+	{
+		FText Message = FText::Format( LOCTEXT( "DatasmithMeshOperations_SetLODs_Max", "Limiting number of reduction settings to max allowed, {0}" ), MAX_STATIC_MESH_LODS );
+		LogWarning( Message );
+	}
+
+	// Limit size of array to MAX_STATIC_MESH_LODS
 	const int32 LODCount = FMath::Min( ReductionSettings.Num(), MAX_STATIC_MESH_LODS );
 	if( LODCount == 0 )
 	{
 		FText OutReason = FText( LOCTEXT( "DatasmithMeshOperations_SetLODs", "No reduction settings. Aborting operation..." ) );
 		LogInfo( OutReason );
-		// #ueent_todo: Remove call to UE_LOG when DataprepLogger is operational
-		UE_LOG(LogDataprep, Log, TEXT("UDataprepSetLODsOperation: %s"), *OutReason.ToString());
 		return;
 	}
 
@@ -163,8 +168,6 @@ void UDataprepSetMaterialOperation::OnExecution_Implementation(const FDataprepCo
 	{
 		FText OutReason = FText( LOCTEXT( "DatasmithMeshOperations_SetMaterial", "No material specified. Aborting operation..." ) );
 		LogInfo( OutReason );
-		// #ueent_todo: Remove call to UE_LOG when DataprepLogger is operational
-		UE_LOG(LogDataprep, Log, TEXT("UDataprepSetMaterialOperation: %s"), *OutReason.ToString());
 		return;
 	}
 
@@ -182,8 +185,6 @@ void UDataprepSubstituteMaterialOperation::OnExecution_Implementation(const FDat
 	{
 		FText OutReason = FText( LOCTEXT( "DatasmithDirProducer_SubstituteMaterial", "No material specified. Aborting operation..." ) );
 		LogInfo( OutReason );
-		// #ueent_todo: Remove call to UE_LOG when DataprepLogger is operational
-		UE_LOG(LogDataprep, Log, TEXT("UDataprepSubstituteMaterialOperation: %s"), *OutReason.ToString());
 		return;
 	}
 
@@ -201,8 +202,6 @@ void UDataprepSubstituteMaterialByTableOperation::OnExecution_Implementation(con
 	{
 		FText OutReason = FText( LOCTEXT( "DatasmithDirProducer_SubstituteMaterialByTable", "No data table specified. Aborting operation..." ) );
 		LogInfo( OutReason );
-		// #ueent_todo: Remove call to UE_LOG when DataprepLogger is operational
-		UE_LOG(LogDataprep, Log, TEXT("UDataprepSubstituteMaterialByTableOperation: %s"), *OutReason.ToString());
 		return;
 	}
 
@@ -299,8 +298,6 @@ void UDataprepSetMeshOperation::OnExecution_Implementation(const FDataprepContex
 	{
 		FText OutReason = FText( LOCTEXT( "DatasmithMeshOperations_SetMesh", "No mesh specified. Aborting operation..." ) );
 		LogInfo( OutReason );
-		// #ueent_todo: Remove call to UE_LOG when DataprepLogger is operational
-		UE_LOG(LogDataprep, Log, TEXT("UDataprepSetMeshOperation: %s"), *OutReason.ToString());
 		return;
 	}
 
diff --git a/Engine/Plugins/Experimental/DataPrepEditor/Source/DataPrepLibraries/Private/DataprepOperations.h b/Engine/Plugins/Experimental/DataPrepEditor/Source/DataPrepLibraries/Private/DataprepOperations.h
--- a/Engine/Plugins/Experimental/DataPrepEditor/Source/DataPrepLibraries/Private/DataprepOperations.h
+++ b/Engine/Plugins/Experimental/DataPrepEditor/Source/DataPrepLibraries/Private/DataprepOperations.h
@@ -58,8 +58,6 @@ public:
 	UPROPERTY(EditAnywhere, BlueprintReadWrite, Category = MeshOperation, Meta = (ToolTip = "Array of LOD reduction settings") )
 	TArray<FDataprepSetLODsReductionSettings> ReductionSettings;
 
-	// #ueent_todo: Limit size of array to MAX_STATIC_MESH_LODS
-
 	//~ Begin UDataprepOperation Interface
 public:
 	virtual FText GetCategory_Implementation() const override
