Fixed a issue on how the Enums type were validated for the dataprep parameterization.

diff --git a/Engine/Plugins/Experimental/DataPrepEditor/Source/DataPrepCore/Private/Parameterization/DataprepParameterization.cpp b/Engine/Plugins/Experimental/DataPrepEditor/Source/DataPrepCore/Private/Parameterization/DataprepParameterization.cpp
--- a/Engine/Plugins/Experimental/DataPrepEditor/Source/DataPrepCore/Private/Parameterization/DataprepParameterization.cpp
+++ b/Engine/Plugins/Experimental/DataPrepEditor/Source/DataPrepCore/Private/Parameterization/DataprepParameterization.cpp
@@ -102,7 +102,7 @@ namespace DataprepParameterization
 			ValueTypeValidationData.Reserve( 3 );
 			ValueTypeValidationData.Add( CurrentProperty->GetClass() );
 			ValueTypeValidationData.Add( CurrentProperty->GetEnum() );
-			ValueTypeValidationData.Add( CurrentProperty->GetUnderlyingProperty() );
+			PopulateValueTypeValidationData( CurrentProperty->GetUnderlyingProperty(), ValueTypeValidationData );
 		}
 	}
 
