Fixed a crash when making a loop with the nodes connection in the dataprep editor.

diff --git a/Engine/Plugins/Experimental/DataPrepEditor/Source/DataPrepCore/Private/DataPrepAsset.cpp b/Engine/Plugins/Experimental/DataPrepEditor/Source/DataPrepCore/Private/DataPrepAsset.cpp
--- a/Engine/Plugins/Experimental/DataPrepEditor/Source/DataPrepCore/Private/DataPrepAsset.cpp
+++ b/Engine/Plugins/Experimental/DataPrepEditor/Source/DataPrepCore/Private/DataPrepAsset.cpp
@@ -249,45 +249,52 @@ void UDataprepAsset::UpdateActions()
 {
 	ActionAssets.Empty(ActionAssets.Num());
 
-	UEdGraphPin* StartNodePin = StartNode->FindPin(UEdGraphSchema_K2::PN_Then, EGPD_Output);
-	if( StartNodePin && StartNodePin->LinkedTo.Num() > 0 )
+	UEdGraphPin* NodeOutPin= StartNode->FindPin(UEdGraphSchema_K2::PN_Then, EGPD_Output);
+	if( NodeOutPin && NodeOutPin->LinkedTo.Num() > 0 )
 	{
-		TSet<UK2Node_DataprepActionCore*> ActionNodesExecuted;
+		TSet<UEdGraphNode*> ActionNodesVisited;
 
-		for( UEdGraphPin* NextNodeInPin = StartNodePin->LinkedTo[0]; NextNodeInPin != nullptr ; )
+		for( UEdGraphPin* NextNodeInPin = NodeOutPin->LinkedTo[0]; NextNodeInPin != nullptr ; )
 		{
 			UEdGraphNode* NextNode = NextNodeInPin->GetOwningNode();
 
-			if(UK2Node_DataprepActionCore* ActionNode = Cast<UK2Node_DataprepActionCore>(NextNode))
+			uint32 NodeHash = GetTypeHash( NextNode );
+			// Break the loop if the node had already been visited
+			if( ActionNodesVisited.FindByHash( NodeHash, NextNode ) )
 			{
-				// Break the loop if the node had already been executed
-				if( ActionNodesExecuted.Find( ActionNode ) )
-				{
-					break;
-				}
+				break;
+			}
+			else
+			{
+				ActionNodesVisited.AddByHash( NodeHash, NextNode );
+			}
 
+			if(UK2Node_DataprepActionCore* ActionNode = Cast<UK2Node_DataprepActionCore>(NextNode))
+			{
+	
 				if(UDataprepActionAsset* DataprepAction = ActionNode->GetDataprepAction())
 				{
 					ActionAssets.Add( DataprepAction );
 				}
 			}
 
-			UEdGraphPin* NextNodeOutPin = NextNode->FindPin( UEdGraphSchema_K2::PN_Then, EGPD_Output );
+			// Look for the next node
+			NodeOutPin = NextNode->FindPin( UEdGraphSchema_K2::PN_Then, EGPD_Output );
 
-			if ( !NextNodeOutPin )
+			if ( !NodeOutPin )
 			{
 				// If we couldn't find a then pin try to get the first output pin as a fallback
 				for ( UEdGraphPin* Pin : NextNode->Pins )
 				{
 					if ( Pin && Pin->PinType.PinCategory == UEdGraphSchema_K2::PC_Exec && Pin->Direction == EGPD_Output )
 					{
-						NextNodeOutPin = Pin;
+						NodeOutPin = Pin;
 						break;
 					}
 				}
 			}
 
-			NextNodeInPin = NextNodeOutPin ? ( NextNodeOutPin->LinkedTo.Num() > 0 ? NextNodeOutPin->LinkedTo[0] : nullptr ) : nullptr;
+			NextNodeInPin = NodeOutPin ? ( NodeOutPin->LinkedTo.Num() > 0 ? NodeOutPin->LinkedTo[0] : nullptr ) : nullptr;
 		}
 	}
 }
