Renamed "Undefined Category" in the Dataprep Palette to "User-defined"

diff --git a/Engine/Plugins/Experimental/DataPrepEditor/Source/DataPrepCore/Private/DataPrepOperation.cpp b/Engine/Plugins/Experimental/DataPrepEditor/Source/DataPrepCore/Private/DataPrepOperation.cpp
--- a/Engine/Plugins/Experimental/DataPrepEditor/Source/DataPrepCore/Private/DataPrepOperation.cpp
+++ b/Engine/Plugins/Experimental/DataPrepEditor/Source/DataPrepCore/Private/DataPrepOperation.cpp
@@ -114,7 +114,7 @@ FText UDataprepOperation::GetTooltip_Implementation() const
 
 FText UDataprepOperation::GetCategory_Implementation() const
 {
-	return FText::FromString( TEXT("Undefined Category") );
+	return LOCTEXT("DefaultOperationCategory", "User-defined");
 }
 
 FText UDataprepOperation::GetAdditionalKeyword_Implementation() const
