Fixing a syntax error in the dataprep float filter filtering criteria.

diff --git a/Engine/Plugins/Experimental/DataPrepEditor/Config/BaseDataprepEditor.ini b/Engine/Plugins/Experimental/DataPrepEditor/Config/BaseDataprepEditor.ini
--- a/Engine/Plugins/Experimental/DataPrepEditor/Config/BaseDataprepEditor.ini
+++ b/Engine/Plugins/Experimental/DataPrepEditor/Config/BaseDataprepEditor.ini
@@ -4,3 +4,4 @@
 +ClassRedirects=(OldName="DataprepRemoveObjectsOperation", NewName="DataprepDeleteObjectsOperation")
 +PropertyRedirects=(OldName="DataprepSetMaterialOperation.MaterialSubstitute",NewName="DataprepSetMaterialOperation.Material")
 +PropertyRedirects=(OldName="DataprepSetMeshOperation.MeshSubstitute",NewName="DataprepSetMeshOperation.StaticMesh")
++EnumRedirects=(OldName="EDataprepFloatMatchType",ValueChanges=(("LessThen","LessThan"),("GreatherThen","GreatherThan")))
diff --git a/Engine/Plugins/Experimental/DataPrepEditor/Source/DataPrepCore/Private/SelectionSystem/DataprepFloatFilter.cpp b/Engine/Plugins/Experimental/DataPrepEditor/Source/DataPrepCore/Private/SelectionSystem/DataprepFloatFilter.cpp
--- a/Engine/Plugins/Experimental/DataPrepEditor/Source/DataPrepCore/Private/SelectionSystem/DataprepFloatFilter.cpp
+++ b/Engine/Plugins/Experimental/DataPrepEditor/Source/DataPrepCore/Private/SelectionSystem/DataprepFloatFilter.cpp
@@ -10,10 +10,10 @@ bool UDataprepFloatFilter::Filter(float Float) const
 {
 	switch ( FloatMatchingCriteria )
 	{
-	case EDataprepFloatMatchType::LessThen:
+	case EDataprepFloatMatchType::LessThan:
 		return Float < EqualValue;
 		break;
-	case EDataprepFloatMatchType::GreatherThen:
+	case EDataprepFloatMatchType::GreatherThan:
 		return Float > EqualValue;
 		break;
 	case EDataprepFloatMatchType::IsNearlyEqual:
diff --git a/Engine/Plugins/Experimental/DataPrepEditor/Source/DataPrepCore/Public/SelectionSystem/DataprepFloatFilter.h b/Engine/Plugins/Experimental/DataPrepEditor/Source/DataPrepCore/Public/SelectionSystem/DataprepFloatFilter.h
--- a/Engine/Plugins/Experimental/DataPrepEditor/Source/DataPrepCore/Public/SelectionSystem/DataprepFloatFilter.h
+++ b/Engine/Plugins/Experimental/DataPrepEditor/Source/DataPrepCore/Public/SelectionSystem/DataprepFloatFilter.h
@@ -15,8 +15,8 @@ class UDataprepFloatFetcher;
 UENUM()
 enum class EDataprepFloatMatchType : uint8
 {
-	LessThen,
-	GreatherThen,
+	LessThan,
+	GreatherThan,
 	IsNearlyEqual
 };
 
