Dataprep - We shouldn't be able to create an instance of an instance in 4.24.
The workflow wasn't fully supported.

diff --git a/Engine/Plugins/Experimental/DataPrepEditor/Source/DataPrepCore/Private/AssetTypeActions_DataprepAssetInterface.cpp b/Engine/Plugins/Experimental/DataPrepEditor/Source/DataPrepCore/Private/AssetTypeActions_DataprepAssetInterface.cpp
--- a/Engine/Plugins/Experimental/DataPrepEditor/Source/DataPrepCore/Private/AssetTypeActions_DataprepAssetInterface.cpp
+++ b/Engine/Plugins/Experimental/DataPrepEditor/Source/DataPrepCore/Private/AssetTypeActions_DataprepAssetInterface.cpp
@@ -109,15 +109,28 @@ void FAssetTypeActions_DataprepAssetInterface::GetActions(const TArray<UObject*>
 		return;
 	}
 
-	MenuBuilder.AddMenuEntry(
-		LOCTEXT("CreateInstance", "Create Instance"),
-		LOCTEXT("CreateInstanceTooltip", "Creates a parameterized Dataprep asset using this Dataprep asset as a base."),
-		FSlateIcon(),
-		FUIAction(
-			FExecuteAction::CreateSP(this, &FAssetTypeActions_DataprepAssetInterface::CreateInstance, DataprepAssetInterfaces),
-			FCanExecuteAction()
-		)
-	);
+	bool bContainsAnInstance  = false;
+	for (UObject* Object : InObjects)
+	{
+		if (Object && Object->GetClass() == UDataprepAssetInstance::StaticClass())
+		{
+			bContainsAnInstance = true;
+			break;
+		}
+	}
+
+	if (!bContainsAnInstance)
+	{
+		MenuBuilder.AddMenuEntry(
+			LOCTEXT("CreateInstance", "Create Instance"),
+			LOCTEXT("CreateInstanceTooltip", "Creates a parameterized Dataprep asset using this Dataprep asset as a base."),
+			FSlateIcon(),
+			FUIAction(
+				FExecuteAction::CreateSP(this, &FAssetTypeActions_DataprepAssetInterface::CreateInstance, DataprepAssetInterfaces),
+				FCanExecuteAction()
+			)
+		);
+	}
 
 	MenuBuilder.AddMenuEntry(
 		LOCTEXT("RunAsset", "Execute"),
