Completed handling of dataprep instance's parent

diff --git a/Engine/Plugins/Experimental/DataPrepEditor/Source/DataPrepCore/Private/AssetTypeActions_DataprepAssetInterface.cpp b/Engine/Plugins/Experimental/DataPrepEditor/Source/DataPrepCore/Private/AssetTypeActions_DataprepAssetInterface.cpp
--- a/Engine/Plugins/Experimental/DataPrepEditor/Source/DataPrepCore/Private/AssetTypeActions_DataprepAssetInterface.cpp
+++ b/Engine/Plugins/Experimental/DataPrepEditor/Source/DataPrepCore/Private/AssetTypeActions_DataprepAssetInterface.cpp
@@ -95,7 +95,11 @@ void FAssetTypeActions_DataprepAssetInterface::ExecuteDataprepAssets(TArray<TWea
 	{
 		if( UDataprepAssetInterface* DataprepAssetInterface = DataprepAssetInterfacePtr.Get() )
 		{
-			UDataprepCoreLibrary::ExecuteWithReporting( DataprepAssetInterface );
+			// Nothing to do if the Dataprep asset does not have any inputs
+			if(DataprepAssetInterface->GetProducers()->GetProducersCount() > 0)
+			{
+				UDataprepCoreLibrary::ExecuteWithReporting( DataprepAssetInterface );
+			}
 		}
 	}
 }
@@ -109,6 +113,8 @@ void FAssetTypeActions_DataprepAssetInterface::GetActions(const TArray<UObject*>
 		return;
 	}
 
+	// #ueent_remark: An instance of an instance is not supported for 4.24.
+	// Do not expose 'Create Instance' menu entry if at least one Dataprep asset is an instance
 	bool bContainsAnInstance  = false;
 	for (UObject* Object : InObjects)
 	{
diff --git a/Engine/Plugins/Experimental/DataPrepEditor/Source/DataPrepCore/Private/DataprepAssetInstance.cpp b/Engine/Plugins/Experimental/DataPrepEditor/Source/DataPrepCore/Private/DataprepAssetInstance.cpp
--- a/Engine/Plugins/Experimental/DataPrepEditor/Source/DataPrepCore/Private/DataprepAssetInstance.cpp
+++ b/Engine/Plugins/Experimental/DataPrepEditor/Source/DataPrepCore/Private/DataprepAssetInstance.cpp
@@ -3,25 +3,53 @@
 #include "DataprepAssetInstance.h"
 
 #include "DataPrepAsset.h"
-#include "DataPrepContentConsumer.h"
 #include "DataprepAssetProducers.h"
+#include "DataPrepContentConsumer.h"
+#include "DataprepCoreLogCategory.h"
 #include "DataprepCorePrivateUtils.h"
 #include "Parameterization/DataprepParameterization.h"
 
+void UDataprepAssetInstance::PostEditChangeProperty(FPropertyChangedEvent& PropertyChangedEvent)
+{
+	// Check if parenting has changed, either through a force delete on the parent or a replace on deletion
+	bool bParentHasChanged = PropertyChangedEvent.ChangeType == EPropertyChangeType::Redirected && PropertyChangedEvent.Property && PropertyChangedEvent.Property->GetName() == TEXT("Parent");
+	
+	if(bParentHasChanged)
+	{
+		if(Parent != nullptr)
+		{
+			UE_LOG( LogDataprepCore, Error, TEXT("UDataprepAssetInstance::PostEditChangeProperty: Reparenting to non null parent is not supported yet.") );
+			Parent = nullptr;
+		}
+
+		// Recreate the ParameterizationInstance property
+		Parameterization = NewObject<UDataprepParameterizationInstance>( this, NAME_None, RF_Public | RF_Transactional );
+
+		// Assign a temporary source parameterization
+		UDataprepParameterization* SourceParameterization = NewObject<UDataprepParameterization>( GetTransientPackage(), FName(), RF_Public );
+		Parameterization->SetParameterizationSource( *SourceParameterization );
+
+		ActionsFromDataprepAsset.Empty();
+	}
+
+	UDataprepAssetInterface::PostEditChangeProperty( PropertyChangedEvent );
+}
+
 // UDataprepAssetInstance =================================================================
 
 void UDataprepAssetInstance::ExecuteRecipe(const TSharedPtr<FDataprepActionContext>& InActionsContext)
 {
-	check( Parent );
-
-	// Doing the parameterizaton
-	TMap<UObject*,UObject*> SourceToCopy;
-	ActionsFromDataprepAsset = GetCopyOfActions( SourceToCopy );
-	Parameterization->ApplyParameterization( SourceToCopy );
+	if(Parent != nullptr)
+	{
+		// Doing the parameterization
+		TMap<UObject*,UObject*> SourceToCopy;
+		ActionsFromDataprepAsset = GetCopyOfActions( SourceToCopy );
+		Parameterization->ApplyParameterization( SourceToCopy );
 
-	ExecuteRecipe_Internal( InActionsContext, ActionsFromDataprepAsset );
+		ExecuteRecipe_Internal( InActionsContext, ActionsFromDataprepAsset );
 
-	ActionsFromDataprepAsset.Empty();
+		ActionsFromDataprepAsset.Empty();
+	}
 }
 
 UObject* UDataprepAssetInstance::GetParameterizationObject()
@@ -31,53 +59,53 @@ UObject* UDataprepAssetInstance::GetParameterizationObject()
 
 TArray<UDataprepActionAsset*> UDataprepAssetInstance::GetCopyOfActions(TMap<UObject*,UObject*>& OutOriginalToCopy) const
 {
-	check( Parent );
-	return Parent->GetCopyOfActions( OutOriginalToCopy );
+	return Parent ? Parent->GetCopyOfActions( OutOriginalToCopy ) : TArray<UDataprepActionAsset*>();
 }
 
 bool UDataprepAssetInstance::SetParent(UDataprepAssetInterface* InParent, bool bNotifyChanges )
 {
+	// #ueent_remark: Setting a null parent is not supported yet.
 	check(InParent);
 
-	// Copy set of producers 
-	if(Inputs != nullptr)
+	if(bNotifyChanges)
 	{
-		Inputs->GetOnChanged().RemoveAll( this );
-		DataprepCorePrivateUtils::DeleteRegisteredAsset( Inputs );
+		Modify();
 	}
-	Inputs = DuplicateObject<UDataprepAssetProducers>( InParent->GetProducers(), this );
 
-	// Copy consumer 
-	if(Output != nullptr)
+	// Copy set of producers if the instance does not have any yet
+	if(Inputs == nullptr || Inputs->GetProducersCount() == 0)
 	{
-		Output->GetOnChanged().RemoveAll( this );
-		DataprepCorePrivateUtils::DeleteRegisteredAsset( Output );
-	}
-	Output = DuplicateObject<UDataprepContentConsumer>( InParent->GetConsumer(), this );
-
-	UDataprepAssetInterface* RealParent = InParent;
-
-	// If InParent is an instance, get the up to the original parent
-	if(UDataprepAssetInstance* InstanceOfInstance = Cast<UDataprepAssetInstance>(InParent))
-	{
-		// #ueent_todo: Copy values from InstanceOfInstance's parameterization to this
-
-		do
+		if(Inputs != nullptr)
 		{
-			RealParent = InstanceOfInstance->Parent;
-			InstanceOfInstance = Cast<UDataprepAssetInstance>(InstanceOfInstance->Parent);
+			Inputs->GetOnChanged().RemoveAll( this );
+			DataprepCorePrivateUtils::DeleteRegisteredAsset( Inputs );
 		}
-		while(InstanceOfInstance != nullptr);
 
-		check(RealParent);
+		Inputs = DuplicateObject<UDataprepAssetProducers>( InParent->GetProducers(), this );
+	}
+
+	// Copy consumer if one is not set yet
+	if(Output == nullptr)
+	{
+		Output = DuplicateObject<UDataprepContentConsumer>( InParent->GetConsumer(), this );
 	}
 
-	Parent = RealParent;
+	// For the time being an instance of instance is not supported
+	Parent = GetRootParent( InParent );
 	
-	check ( Parent->GetClass() == UDataprepAsset::StaticClass() );
 	Parameterization = NewObject<UDataprepParameterizationInstance>( this, NAME_None, RF_Public | RF_Transactional );
-	Parameterization->SetParameterizationSource( *( static_cast<UDataprepAsset*>( Parent )->GetDataprepParameterization() ) );
-	
+
+	// #ueent_nextversion: Get parameterization from root parent and copy values from actual parent if parent is an instance
+	if( UDataprepAsset* RootParent = Cast<UDataprepAsset>(Parent) )
+	{
+		Parameterization->SetParameterizationSource( *( RootParent->GetDataprepParameterization() ) );
+	}
+	else
+	{
+		// Assign a temporary source parameterization
+		UDataprepParameterization* SourceParameterization = NewObject<UDataprepParameterization>( GetTransientPackage(), FName(), RF_Public );
+		Parameterization->SetParameterizationSource( *SourceParameterization );
+	}
 
 	if(bNotifyChanges)
 	{
@@ -88,3 +116,24 @@ bool UDataprepAssetInstance::SetParent(UDataprepAssetInterface* InParent, bool b
 
 	return true;
 }
+
+UDataprepAsset* UDataprepAssetInstance::GetRootParent(UDataprepAssetInterface* InParent)
+{
+	UDataprepAsset* RootParent = Cast<UDataprepAsset>(InParent);
+
+	while(RootParent == nullptr)
+	{
+		if(UDataprepAssetInstance* InstanceOfInstance = Cast<UDataprepAssetInstance>(InParent))
+		{
+			if(InstanceOfInstance->Parent == nullptr)
+			{
+				break;
+			}
+
+			InParent = Cast<UDataprepAssetInstance>(InstanceOfInstance->Parent);
+			RootParent = Cast<UDataprepAsset>(InParent);
+		}
+	}
+
+	return RootParent;
+}
\ No newline at end of file
diff --git a/Engine/Plugins/Experimental/DataPrepEditor/Source/DataPrepCore/Private/Parameterization/DataprepParameterization.cpp b/Engine/Plugins/Experimental/DataPrepEditor/Source/DataPrepCore/Private/Parameterization/DataprepParameterization.cpp
--- a/Engine/Plugins/Experimental/DataPrepEditor/Source/DataPrepCore/Private/Parameterization/DataprepParameterization.cpp
+++ b/Engine/Plugins/Experimental/DataPrepEditor/Source/DataPrepCore/Private/Parameterization/DataprepParameterization.cpp
@@ -3,6 +3,7 @@
 #include "Parameterization/DataprepParameterization.h"
 
 #include "DataPrepAsset.h"
+#include "DataPrepAssetInstance.h"
 #include "DataprepParameterizableObject.h"
 #include "DataprepParameterizationArchive.h"
 #include "Parameterization/DataprepParameterizationUtils.h"
@@ -1366,6 +1367,12 @@ void UDataprepParameterizationInstance::PostLoad()
 {
 	if ( !HasAnyFlags( RF_ClassDefaultObject | RF_NeedLoad ) )
 	{
+		// #ueent_hotfix: If source is null, the parent of the DataprepAssetInstance is null. So recreate a temporary source parameterization
+		if(SourceParameterization == nullptr)
+		{
+			ensure( Cast<UDataprepAssetInstance>( GetOuter()) && Cast<UDataprepAssetInstance>( GetOuter())->GetParent() == nullptr );
+			SourceParameterization = NewObject<UDataprepParameterization>( GetTransientPackage(), FName(), RF_Public );
+		}
 		SetFlags( RF_Public );
 		LoadParameterization();
 		SetupCallbacksFromSourceParameterisation();
diff --git a/Engine/Plugins/Experimental/DataPrepEditor/Source/DataPrepCore/Public/DataprepAssetInstance.h b/Engine/Plugins/Experimental/DataPrepEditor/Source/DataPrepCore/Public/DataprepAssetInstance.h
--- a/Engine/Plugins/Experimental/DataPrepEditor/Source/DataPrepCore/Public/DataprepAssetInstance.h
+++ b/Engine/Plugins/Experimental/DataPrepEditor/Source/DataPrepCore/Public/DataprepAssetInstance.h
@@ -9,6 +9,7 @@
 #include "DataprepAssetInstance.generated.h"
 
 class UDataprepActionAsset;
+class UDataprepAsset;
 class UDataprepParameterizationInstance;
 
 /**
@@ -30,6 +31,10 @@ public:
 
 	virtual ~UDataprepAssetInstance() = default;
 
+	// UObject interface
+	virtual void PostEditChangeProperty( struct FPropertyChangedEvent& PropertyChangedEvent) override;
+	// End of UObject interface
+
 	// UDataprepAssetInterface interface
 	virtual void ExecuteRecipe( const TSharedPtr<FDataprepActionContext>& InActionsContext ) override;
 	virtual bool HasActions() const override { return Parent ? Parent->HasActions() : false; }
@@ -57,11 +62,13 @@ protected:
 	UPROPERTY()
 	UDataprepParameterizationInstance* Parameterization;
 
-
-protected:
 	/** Delegate broadcasted when the consumer or one of the producers has changed */
 	FOnDataprepAssetInstanceChanged OnParentChanged;
 
+private:
+	/** Returns a pointer to the UDataprepAsset */
+	UDataprepAsset* GetRootParent(UDataprepAssetInterface* InParent);
+
 private:
 	UPROPERTY()
 	TArray<UDataprepActionAsset*> ActionsFromDataprepAsset;
