Fixed build breakage

diff --git a/Engine/Plugins/Experimental/DataPrepEditor/Source/DataPrepCore/Private/AssetTypeActions_DataprepAssetInterface.cpp b/Engine/Plugins/Experimental/DataPrepEditor/Source/DataPrepCore/Private/AssetTypeActions_DataprepAssetInterface.cpp
--- a/Engine/Plugins/Experimental/DataPrepEditor/Source/DataPrepCore/Private/AssetTypeActions_DataprepAssetInterface.cpp
+++ b/Engine/Plugins/Experimental/DataPrepEditor/Source/DataPrepCore/Private/AssetTypeActions_DataprepAssetInterface.cpp
@@ -3,6 +3,7 @@
 #include "AssetTypeActions_DataprepAssetInterface.h"
 
 #include "DataprepAssetInstance.h"
+#include "DataprepAssetProducers.h"
 #include "DataprepCoreLibrary.h"
 #include "DataPrepCoreModule.h"
 #include "DataPrepFactories.h"
