Fixed some issues with the dataprep parameterization
Like the dataprep not always properly syncing mofidified properties.

diff --git a/Engine/Plugins/Experimental/DataPrepEditor/Source/DataPrepCore/Private/DataPrepAsset.cpp b/Engine/Plugins/Experimental/DataPrepEditor/Source/DataPrepCore/Private/DataPrepAsset.cpp
--- a/Engine/Plugins/Experimental/DataPrepEditor/Source/DataPrepCore/Private/DataPrepAsset.cpp
+++ b/Engine/Plugins/Experimental/DataPrepEditor/Source/DataPrepCore/Private/DataPrepAsset.cpp
@@ -231,9 +231,9 @@ void UDataprepAsset::RemoveObjectPropertyFromParameterization(UDataprepParameter
 	Parameterization->RemoveBindedObjectProperty( Object, InPropertyChain );
 }
 
-void UDataprepAsset::GetExistingParameterNamesForType(UProperty* Property, TSet<FString>& OutValidExistingNames, TSet<FString>& OutInvalidNames) const
+void UDataprepAsset::GetExistingParameterNamesForType(UProperty* Property, bool bIsDescribingFullProperty, TSet<FString>& OutValidExistingNames, TSet<FString>& OutInvalidNames) const
 {
-	Parameterization->GetExistingParameterNamesForType( Property, OutValidExistingNames, OutInvalidNames );
+	Parameterization->GetExistingParameterNamesForType( Property, bIsDescribingFullProperty, OutValidExistingNames, OutInvalidNames );
 }
 
 void UDataprepAsset::OnDataprepBlueprintChanged( UBlueprint* InBlueprint )
@@ -271,8 +271,7 @@ void UDataprepAsset::UpdateActions()
 
 			if(UK2Node_DataprepActionCore* ActionNode = Cast<UK2Node_DataprepActionCore>(NextNode))
 			{
-	
-				if(UDataprepActionAsset* DataprepAction = ActionNode->GetDataprepAction())
+				if( UDataprepActionAsset* DataprepAction = ActionNode->GetDataprepAction() )
 				{
 					ActionAssets.Add( DataprepAction );
 				}
diff --git a/Engine/Plugins/Experimental/DataPrepEditor/Source/DataPrepCore/Private/Parameterization/DataprepParameterization.cpp b/Engine/Plugins/Experimental/DataPrepEditor/Source/DataPrepCore/Private/Parameterization/DataprepParameterization.cpp
--- a/Engine/Plugins/Experimental/DataPrepEditor/Source/DataPrepCore/Private/Parameterization/DataprepParameterization.cpp
+++ b/Engine/Plugins/Experimental/DataPrepEditor/Source/DataPrepCore/Private/Parameterization/DataprepParameterization.cpp
@@ -240,7 +240,12 @@ namespace DataprepParameterization
 			return GetAddressOf( static_cast<const UMapProperty&>( Property ), BaseAddress, ContainerIndex );
 		}
 
-		return Property.ContainerPtrToValuePtr<void*>(BaseAddress, FMath::Max( ContainerIndex, 0 ) );
+		if ( Property.ArrayDim > ContainerIndex )
+		{
+			return Property.ContainerPtrToValuePtr<void*>( BaseAddress, FMath::Max( ContainerIndex, 0 ) );
+		}
+		
+		return nullptr;
 	}
 
 	/**
@@ -453,7 +458,7 @@ namespace DataprepParameterization
 		return GetPropertyFromBinding( Binding, DummyPointer );
 	}
 
-	void CopyCompleteValue(UProperty& DestinationProperty, void* DestinationAddress, UProperty& SourceProperty, void* SourceAddress)
+	void CopyValue(UProperty& DestinationProperty, void* DestinationAddress, UProperty& SourceProperty, void* SourceAddress)
 	{
 		UClass* ProperyClass = DestinationProperty.GetClass();
 		// We only support copying value of properties when they are from the same class (this is not a warranty that this is safe, it's only a validation heuristic)
@@ -465,6 +470,11 @@ namespace DataprepParameterization
 			const bool bSourceValue = static_cast<UBoolProperty&>( SourceProperty ).GetPropertyValue( SourceAddress );
 			static_cast<UBoolProperty&>( DestinationProperty ).SetPropertyValue( DestinationAddress, bSourceValue );
 		}
+		else if ( DestinationProperty.ArrayDim != SourceProperty.ArrayDim )
+		{
+			UProperty& SmallerProperty = DestinationProperty.ArrayDim > SourceProperty.ArrayDim ? SourceProperty : DestinationProperty;
+			SmallerProperty.CopySingleValue( DestinationAddress, SourceAddress );
+		}
 		else
 		{
 			DestinationProperty.CopyCompleteValue( DestinationAddress, SourceAddress );
@@ -537,7 +547,7 @@ bool UDataprepParameterizationBindings::HasBindingsForParameter(const FName& Par
 	return NameToBindings.Contains( ParameterName );
 }
 
-void UDataprepParameterizationBindings::Add(const TSharedRef<FDataprepParameterizationBinding>& Binding, const FName& ParamerterName)
+void UDataprepParameterizationBindings::Add(const TSharedRef<FDataprepParameterizationBinding>& Binding, const FName& ParamerterName, FSetOfBinding& OutBindingsContainedByNewBinding)
 {
 	Modify();
 
@@ -554,7 +564,38 @@ void UDataprepParameterizationBindings::Add(const TSharedRef<FDataprepParameteri
 
 	BindingToParameterName.AddByHash( BindingHash, Binding, ParamerterName );
 	NameToBindings.FindOrAdd( ParamerterName ).AddByHash( BindingHash, Binding );
-	ObjectToBindings.FindOrAdd( Binding->ObjectBinded ).AddByHash( BindingHash, Binding );
+
+	FSetOfBinding& BindingsFromSameObject = ObjectToBindings.FindOrAdd( Binding->ObjectBinded );
+		
+	for ( TSharedRef<FDataprepParameterizationBinding>& PossibleSubBinding : BindingsFromSameObject )
+	{
+		if ( PossibleSubBinding->PropertyChain.Num() >= Binding->PropertyChain.Num() )
+		{
+			TArray<FDataprepPropertyLink> PropertyChain = PossibleSubBinding->PropertyChain;
+			while ( PropertyChain.Num() > Binding->PropertyChain.Num() )
+			{
+				PropertyChain.RemoveAt( PropertyChain.Num() - 1, 1, false );
+			}
+
+			if ( PropertyChain == Binding->PropertyChain )
+			{
+				OutBindingsContainedByNewBinding.Add( PossibleSubBinding );
+			}
+			else if ( PropertyChain.Num() > 0 )
+			{
+				PropertyChain.Last().ContainerIndex = INDEX_NONE;
+				if ( PropertyChain == Binding->PropertyChain )
+				{
+					OutBindingsContainedByNewBinding.Add( PossibleSubBinding );
+				}
+			}
+		}
+		
+	}
+
+	BindingsFromSameObject.AddByHash( BindingHash, Binding );
+
+	
 }
 
 FName UDataprepParameterizationBindings::RemoveBinding(const TSharedRef<FDataprepParameterizationBinding>& Binding)
@@ -625,6 +666,36 @@ TSet<FName> UDataprepParameterizationBindings::RemoveAllBindingsFromObject(UData
 	return {};
 }
 
+TSharedPtr<FDataprepParameterizationBinding> UDataprepParameterizationBindings::GetContainingBinding(const TSharedRef<FDataprepParameterizationBinding>& Binding) const
+{
+	if ( Binding->ObjectBinded )
+	{
+		 TSharedRef<FDataprepParameterizationBinding> PossibleContainingBinding = MakeShared<FDataprepParameterizationBinding>( Binding.Get() );
+
+		while ( PossibleContainingBinding->PropertyChain.Num() > 0 )
+		{
+			if ( BindingToParameterName.Contains( PossibleContainingBinding ) )
+			{
+				PossibleContainingBinding->ValueTypeValidationData.Empty();
+				DataprepParameterization::PopulateValueTypeValidationData( PossibleContainingBinding->PropertyChain.Last().CachedProperty.Get(), PossibleContainingBinding->ValueTypeValidationData );
+				return PossibleContainingBinding;
+			}
+
+			int32& ContainerIndex = PossibleContainingBinding->PropertyChain.Last().ContainerIndex;
+			if ( ContainerIndex != INDEX_NONE )
+			{
+				ContainerIndex = INDEX_NONE;
+			}
+			else
+			{
+				PossibleContainingBinding->PropertyChain.Pop( false );
+			}
+		}
+	}
+
+	return {};
+}
+
 const UDataprepParameterizationBindings::FBindingToParameterNameMap& UDataprepParameterizationBindings::GetBindingToParameterName() const
 {
 	return BindingToParameterName;
@@ -847,7 +918,7 @@ void UDataprepParameterization::PostEditUndo()
 
 	OnTellInstancesToReloadTheirSerializedData.Broadcast();
 
-	CastChecked<UDataprepAsset>( GetOuter() )->OnParameterizationStatusForObjectsChanged.Broadcast( nullptr );
+	CastChecked<UDataprepAsset>( GetOuter() )->OnParameterizedObjectsChanged.Broadcast( nullptr );
 }
 
 void UDataprepParameterization::OnObjectModified(UObject* Object)
@@ -883,38 +954,49 @@ bool UDataprepParameterization::BindObjectProperty(UDataprepParameterizableObjec
 		}
 
 		bool bBindingWasAdded = false;
+		bool bAddingFullProperty = PropertyChain.Last().ContainerIndex == INDEX_NONE;
 
 		// We expect the chain to have a valid chain of cached property before inserting the binding
 		if ( DataprepParameterization::GetDeepestLevelOfValidCache( Binding.Get(), AddressOfTheValueFromBinding) == PropertyChain.Num() - 1 )
 		{
+			UDataprepParameterizationBindings::FSetOfBinding BindingsToRemove;
+
 			if ( UProperty** PropertyPtr = NameToParameterizationProperty.Find( Name ) )
 			{
-				UProperty* Property = *PropertyPtr;
-
+				UProperty* PropertyFromParameterization = *PropertyPtr;
+				UProperty* PropertyFromBinding = PropertyChain.Last().CachedProperty.Get();
 				// Ensure that the properties are compatible 
-				TArray<UObject*> ValueTypeValidationData;
-				DataprepParameterization::PopulateValueTypeValidationData( Property, ValueTypeValidationData );
-				if ( ValueTypeValidationData == Binding->ValueTypeValidationData )
+				if ( !bAddingFullProperty || PropertyFromParameterization->ArrayDim == PropertyFromBinding->ArrayDim )
 				{
-					Modify();
-					BindingsContainer->Add( Binding, Name );
-					bBindingWasAdded = true;
+					TArray<UObject*> ValueTypeValidationData;
+					DataprepParameterization::PopulateValueTypeValidationData( PropertyFromParameterization, ValueTypeValidationData );
+					if ( ValueTypeValidationData == Binding->ValueTypeValidationData )
+					{
+						Modify();
+						BindingsContainer->Add( Binding, Name, BindingsToRemove );
+						bBindingWasAdded = true;
+					}
 				}
 			}
 			else
 			{
 				Modify();
 
-				BindingsContainer->Add( Binding, Name );
+				BindingsContainer->Add( Binding, Name, BindingsToRemove );
 			
-				UProperty* SourceProperty = PropertyChain.Last().CachedProperty.Get();
+				UProperty* PropertyFromBinding = PropertyChain.Last().CachedProperty.Get();
 
 				// The validation we did with GetDeepestLevelOfValidCache ensure us that the property ptr is valid
-				UProperty* NewProperty = AddPropertyToClass( Name, *SourceProperty );
+				UProperty* NewProperty = AddPropertyToClass( Name, *PropertyFromBinding, bAddingFullProperty );
 
 				bClassNeedUpdate = true;
 				bBindingWasAdded = true;
 			}
+
+			for ( TSharedRef<FDataprepParameterizationBinding>& BindingToRemove : BindingsToRemove )
+			{
+				RemoveBinding( BindingToRemove, bClassNeedUpdate );
+			}
 		}
 
 		if ( bClassNeedUpdate )
@@ -928,7 +1010,7 @@ bool UDataprepParameterization::BindObjectProperty(UDataprepParameterizableObjec
 
 			TSet<UObject*> Objects;
 			Objects.Add( Object );
-			CastChecked<UDataprepAsset>( GetOuter() )->OnParameterizationStatusForObjectsChanged.Broadcast( &Objects );
+			CastChecked<UDataprepAsset>( GetOuter() )->OnParameterizedObjectsChanged.Broadcast( &Objects );
 		}
 	}
 
@@ -963,7 +1045,7 @@ void UDataprepParameterization::RemoveBindedObjectProperty(UDataprepParameteriza
 
 		TSet<UObject*> Objects;
 		Objects.Add( Object );
-		CastChecked<UDataprepAsset>( GetOuter() )->OnParameterizationStatusForObjectsChanged.Broadcast( &Objects );
+		CastChecked<UDataprepAsset>( GetOuter() )->OnParameterizedObjectsChanged.Broadcast( &Objects );
 	}
 }
 
@@ -995,29 +1077,33 @@ void UDataprepParameterization::RemoveBindingFromObjects(TArray<UDataprepParamet
 		UpdateClass();
 	}
 
-	CastChecked<UDataprepAsset>( GetOuter() )->OnParameterizationStatusForObjectsChanged.Broadcast( &UniqueObjects );
+	CastChecked<UDataprepAsset>( GetOuter() )->OnParameterizedObjectsChanged.Broadcast( &UniqueObjects );
 }
 
 void UDataprepParameterization::UpdateParameterizationFromBinding(const TSharedRef<FDataprepParameterizationBinding>& Binding)
 {
-	FName ParameterModified = BindingsContainer->GetParameterNameForBinding(Binding);
-	UProperty* ParameterizationProperty = NameToParameterizationProperty.FindRef(ParameterModified);
-	if (ParameterizationProperty)
+	FName ParameterModified = BindingsContainer->GetParameterNameForBinding( Binding );
+	UProperty* ParameterizationProperty = NameToParameterizationProperty.FindRef( ParameterModified );
+	if ( ParameterizationProperty )
 	{
 		void* AddressOfObjectValue = nullptr;
-		if (UProperty * ObjectProperty = DataprepParameterization::GetPropertyFromBinding(Binding.Get(), AddressOfObjectValue))
+		if ( UProperty * ObjectProperty = DataprepParameterization::GetPropertyFromBinding( Binding.Get(), AddressOfObjectValue ) )
 		{
 			Modify();
-			void* AddressOfParameterizationValue = DataprepParameterization::GetAddressOf(*ParameterizationProperty, DefaultParameterisation, INDEX_NONE);
-			DataprepParameterization::CopyCompleteValue(*ParameterizationProperty, AddressOfParameterizationValue, *ObjectProperty, AddressOfObjectValue);
+			void* AddressOfParameterizationValue = DataprepParameterization::GetAddressOf( *ParameterizationProperty, DefaultParameterisation, INDEX_NONE );
+			DataprepParameterization::CopyValue( *ParameterizationProperty, AddressOfParameterizationValue, *ObjectProperty, AddressOfObjectValue );
 
 			// Post edit the default parameterization
 			FEditPropertyChain EditChain;
-			EditChain.AddHead(ParameterizationProperty);
-			EditChain.SetActivePropertyNode(ParameterizationProperty);
-			FPropertyChangedEvent EditPropertyChangeEvent(ParameterizationProperty, EPropertyChangeType::ValueSet);
-			FPropertyChangedChainEvent EditChangeChainEvent(EditChain, EditPropertyChangeEvent);
-			DefaultParameterisation->PostEditChangeChainProperty(EditChangeChainEvent);
+			EditChain.AddHead( ParameterizationProperty );
+			EditChain.SetActivePropertyNode( ParameterizationProperty );
+			FPropertyChangedEvent EditPropertyChangeEvent( ParameterizationProperty, EPropertyChangeType::ValueSet );
+			FPropertyChangedChainEvent EditChangeChainEvent( EditChain, EditPropertyChangeEvent );
+			DefaultParameterisation->PostEditChangeChainProperty( EditChangeChainEvent );
+
+			TSet<UObject*> Objects;
+			Objects.Add( DefaultParameterisation );
+			CastChecked<UDataprepAsset>( GetOuter() )->OnParameterizedObjectsChanged.Broadcast( &Objects );
 		}
 	}
 }
@@ -1030,19 +1116,20 @@ void UDataprepParameterization::OnObjectPostEdit(UDataprepParameterizableObject*
 		{
 			PushParametrizationValueToBindings( PropertyChain[0].PropertyName );
 		}
-		else if ( bool( ChangeType & EPropertyChangeType::ValueSet ) )
-		{
-			TSharedRef<FDataprepParameterizationBinding> Binding = MakeShared<FDataprepParameterizationBinding>( Object, PropertyChain );
-			UpdateParameterizationFromBinding( Binding );
-		}
-		else if ( const UDataprepParameterizationBindings::FSetOfBinding* Bindings = BindingsContainer->GetBindingsFromObject( Object ) )
+		else
 		{
-			// todo for container ChangeType
+			TSharedRef<FDataprepParameterizationBinding> BindingForModifiedProperty = MakeShared<FDataprepParameterizationBinding>( Object, PropertyChain );
+
+			TSharedPtr<FDataprepParameterizationBinding> Binding = BindingsContainer->GetContainingBinding( BindingForModifiedProperty );
+			if ( Binding )
+			{
+				UpdateParameterizationFromBinding( Binding.ToSharedRef() );
+			}
 		}
 	}
 }
 
-void UDataprepParameterization::GetExistingParameterNamesForType(UProperty* Property, TSet<FString>& OutValidExistingNames, TSet<FString>& OutInvalidNames) const
+void UDataprepParameterization::GetExistingParameterNamesForType(UProperty* Property, bool bIsDescribingFullProperty, TSet<FString>& OutValidExistingNames, TSet<FString>& OutInvalidNames) const
 {
 	OutValidExistingNames.Empty( NameToParameterizationProperty.Num() );
 	OutInvalidNames.Empty( NameToParameterizationProperty.Num() );
@@ -1052,7 +1139,7 @@ void UDataprepParameterization::GetExistingParameterNamesForType(UProperty* Prop
 		if ( UProperty* ParameterizationProperty = Pair.Value )
 		{
 			bool bWasAdded = false;
-			if ( ParameterizationProperty->GetClass() == Property->GetClass() )
+			if ( ParameterizationProperty->GetClass() == Property->GetClass() && ( !bIsDescribingFullProperty || ParameterizationProperty->ArrayDim == Property->ArrayDim ) )
 			{
 				TArray<UObject*> ValidationDataForParameterizationProperty;
 				DataprepParameterization::PopulateValueTypeValidationData( ParameterizationProperty, ValidationDataForParameterizationProperty );
@@ -1094,7 +1181,7 @@ void UDataprepParameterization::GenerateClass()
 			// Need to manually call Link to fix-up some data (such as the C++ property flags) that are only set during Link
 			{
 				FArchive Ar;
-				NewProperty->LinkWithoutChangingOffset(Ar);
+				NewProperty->LinkWithoutChangingOffset( Ar );
 			}
 
 			CustomContainerClass->AddCppProperty( NewProperty );
@@ -1145,7 +1232,7 @@ void UDataprepParameterization::LoadParameterization()
 
 			if ( PropertyFromChain && !PropertyFromParameterizationClass )
 			{
-				UProperty* NewProperty = AddPropertyToClass( BindingName, *PropertyFromChain );
+				UProperty* NewProperty = AddPropertyToClass( BindingName, *PropertyFromChain, Binding.Key->PropertyChain.Last().ContainerIndex == INDEX_NONE );
 			}
 			else if ( !PropertyFromChain || PropertyFromChain->GetClass() != (*PropertyFromParameterizationClass)->GetClass() )
 			{
@@ -1181,7 +1268,7 @@ void UDataprepParameterization::LoadParameterization()
 
 		if ( ObjectsToNotify.Num() > 0 )
 		{
-			CastChecked<UDataprepAsset>( GetOuter() )->OnParameterizationStatusForObjectsChanged.Broadcast( &ObjectsToNotify );
+			CastChecked<UDataprepAsset>( GetOuter() )->OnParameterizedObjectsChanged.Broadcast( &ObjectsToNotify );
 		}
 	}
 }
@@ -1263,7 +1350,7 @@ void UDataprepParameterization::DoReinstancing(UClass* OldClass, bool bMigrateDa
 	}
 }
 
-UProperty* UDataprepParameterization::AddPropertyToClass(const FName& ParameterisationPropertyName, UProperty& Property)
+UProperty* UDataprepParameterization::AddPropertyToClass(const FName& ParameterisationPropertyName, UProperty& Property, bool bAddFullProperty)
 {
 	if ( !NameToParameterizationProperty.Find( ParameterisationPropertyName ) )
 	{
@@ -1271,6 +1358,11 @@ UProperty* UDataprepParameterization::AddPropertyToClass(const FName& Parameteri
 		NewProperty->SetFlags( RF_Transient );
 		NewProperty->PropertyFlags = CPF_Edit | CPF_NonTransactional;
 
+		if ( !bAddFullProperty )
+		{
+			NewProperty->ArrayDim = 1;
+		}
+
 		// Need to manually call Link to fix-up some data (such as the C++ property flags) that are only set during Link
 		{
 			FArchive Ar;
@@ -1297,13 +1389,15 @@ void UDataprepParameterization::PushParametrizationValueToBindings(FName Paramet
 			{
 				TSet<TSharedRef<FDataprepParameterizationBinding>> BindingToRemove;
 
+				TSet<UObject*> ObjectsModified;
 				for ( const TSharedRef<FDataprepParameterizationBinding>& Binding : *Bindings )
 				{
 					void* AddressOfBindingValue = nullptr;
 					if ( UProperty* BindingProperty = DataprepParameterization::GetPropertyFromBinding( Binding.Get(), AddressOfBindingValue ) )
 					{
+						ObjectsModified.Add( Binding->ObjectBinded );
 						Binding->ObjectBinded->Modify();
-						DataprepParameterization::CopyCompleteValue( *BindingProperty, AddressOfBindingValue, *ParameterizationPropeterty, AddressOfParameterValue );
+						DataprepParameterization::CopyValue( *BindingProperty, AddressOfBindingValue, *ParameterizationPropeterty, AddressOfParameterValue );
 					}
 					else
 					{
@@ -1311,6 +1405,9 @@ void UDataprepParameterization::PushParametrizationValueToBindings(FName Paramet
 					}
 				}
 
+				
+				CastChecked<UDataprepAsset>( GetOuter() )->OnParameterizedObjectsChanged.Broadcast( &ObjectsModified );
+
 				TSet<UObject*> ObjectsToNotify;
 				bool bClassNeedUpdate = false;
 
@@ -1332,7 +1429,7 @@ void UDataprepParameterization::PushParametrizationValueToBindings(FName Paramet
 
 				if ( ObjectsToNotify.Num() > 0 )
 				{
-					CastChecked<UDataprepAsset>( GetOuter() )->OnParameterizationStatusForObjectsChanged.Broadcast( &ObjectsToNotify );
+					CastChecked<UDataprepAsset>( GetOuter() )->OnParameterizedObjectsChanged.Broadcast( &ObjectsToNotify );
 				}
 			}
 		}
@@ -1421,7 +1518,7 @@ void UDataprepParameterizationInstance::ApplyParameterization(const TMap<UObject
 			{
 				UProperty* ParameterizationProperty = FindObjectFast<UProperty>( SourceParameterization->CustomContainerClass, BindingPair.Value );
 				void* ParameterizationAddress =  DataprepParameterization::GetAddressOf( *ParameterizationProperty, ParameterizationInstance, INDEX_NONE );
-				DataprepParameterization::CopyCompleteValue( *DestinationProperty, DestinationAddress, *ParameterizationProperty, ParameterizationAddress );
+				DataprepParameterization::CopyValue( *DestinationProperty, DestinationAddress, *ParameterizationProperty, ParameterizationAddress );
 			}
 		}
 	}
diff --git a/Engine/Plugins/Experimental/DataPrepEditor/Source/DataPrepCore/Private/Parameterization/DataprepParameterization.h b/Engine/Plugins/Experimental/DataPrepEditor/Source/DataPrepCore/Private/Parameterization/DataprepParameterization.h
--- a/Engine/Plugins/Experimental/DataPrepEditor/Source/DataPrepCore/Private/Parameterization/DataprepParameterization.h
+++ b/Engine/Plugins/Experimental/DataPrepEditor/Source/DataPrepCore/Private/Parameterization/DataprepParameterization.h
@@ -134,7 +134,7 @@ public:
 	/**
 	 * Add a binding and map it to the parameter
 	 */
-	void Add(const TSharedRef<FDataprepParameterizationBinding>& Binding, const FName& ParameterName);
+	void Add(const TSharedRef<FDataprepParameterizationBinding>& Binding, const FName& ParameterName, FSetOfBinding& OutBindingsContainedByNewBinding);
 
 	/**
 	 * Remove a binding.
@@ -148,6 +148,11 @@ public:
 	 */
 	TSet<FName> RemoveAllBindingsFromObject(UDataprepParameterizableObject* Object);
 
+	/**
+	 * @return A valid pointer if the binding is a part of an existing binding
+	 */
+	TSharedPtr<FDataprepParameterizationBinding> GetContainingBinding(const TSharedRef<FDataprepParameterizationBinding>& Binding) const;
+
 	const FBindingToParameterNameMap& GetBindingToParameterName() const;
 
 	// UObject interface
@@ -212,7 +217,7 @@ public:
 
 	void OnObjectPostEdit(UDataprepParameterizableObject* Object, const TArray<FDataprepPropertyLink>& PropertyChain, EPropertyChangeType::Type ChangeType);
 
-	void GetExistingParameterNamesForType(UProperty* Property,TSet<FString>& OutValidExistingNames, TSet<FString>& OutInvalidNames) const;
+	void GetExistingParameterNamesForType(UProperty* Property, bool bIsDescribingFullProperty, TSet<FString>& OutValidExistingNames, TSet<FString>& OutInvalidNames) const;
 
 private:
 
@@ -250,9 +255,8 @@ private:
 
 	/**
 	 * Try adding a binded property to the parameterization class
-	 * @return false if the binding is no more valid
 	 */
-	UProperty* AddPropertyToClass(const FName& ParameterisationPropertyName, UProperty& Property);
+	UProperty* AddPropertyToClass(const FName& ParameterisationPropertyName, UProperty& Property, bool bAddFullProperty);
 
 	/**
 	 * Get a new value for the parameterization from it's associated binding
diff --git a/Engine/Plugins/Experimental/DataPrepEditor/Source/DataPrepCore/Private/Parameterization/DataprepParameterizationUtils.cpp b/Engine/Plugins/Experimental/DataPrepEditor/Source/DataPrepCore/Private/Parameterization/DataprepParameterizationUtils.cpp
--- a/Engine/Plugins/Experimental/DataPrepEditor/Source/DataPrepCore/Private/Parameterization/DataprepParameterizationUtils.cpp
+++ b/Engine/Plugins/Experimental/DataPrepEditor/Source/DataPrepCore/Private/Parameterization/DataprepParameterizationUtils.cpp
@@ -73,7 +73,7 @@ TArray<FDataprepPropertyLink> FDataprepParameterizationUtils::MakePropertyChain(
 				// We manipulate a bit the chain to store the property inside a container in a special way. So that we can 
 				if ( DataprepParameterizationUtils::IsAContainerProperty( ParentProperty ) )
 				{
-					PropertyChain.Emplace( Property, Property->GetFName(), 0);
+					PropertyChain.Emplace( Property, Property->GetFName(), INDEX_NONE );
 					PropertyChain.Emplace( ParentProperty, ParentProperty->GetFName(), CurrentHandle->GetIndexInArray() );
 					CurrentHandle = ParentHandle->GetParentHandle();
 					bWasProccess = true;
@@ -87,8 +87,25 @@ TArray<FDataprepPropertyLink> FDataprepParameterizationUtils::MakePropertyChain(
 			}
 
 			if ( CurrentHandle )
-			{ 
-				Property = CurrentHandle->GetProperty();
+			{
+				UProperty* NextProperty = CurrentHandle->GetProperty();
+
+				// Deal with the case of a property of a raw c++ array
+				while ( CurrentHandle && NextProperty == Property )
+				{
+					// Skip the new current handle as it point on the same property
+					CurrentHandle = CurrentHandle->GetParentHandle();
+					if ( CurrentHandle.IsValid() )
+					{
+						NextProperty = CurrentHandle->GetProperty();
+					}
+					else
+					{
+						NextProperty = nullptr;
+					}
+				}
+
+				Property = NextProperty;
 			}
 		}
 		
@@ -109,6 +126,7 @@ TArray<FDataprepPropertyLink> FDataprepParameterizationUtils::MakePropertyChain(
 
 TArray<FDataprepPropertyLink> FDataprepParameterizationUtils::MakePropertyChain(FPropertyChangedChainEvent& PropertyChangedEvent)
 {
+	// This implementation is base on what the property editor does in general
 	const FEditPropertyChain& EditPropertyChain = PropertyChangedEvent.PropertyChain;
 
 	TArray<FDataprepPropertyLink> DataprepPropertyChain;
@@ -124,6 +142,37 @@ TArray<FDataprepPropertyLink> FDataprepParameterizationUtils::MakePropertyChain(
 		{ 
 			int32 ContainerIndex = PropertyChangedEvent.GetArrayIndex( Property->GetName() );
 			DataprepPropertyChain.Emplace( Property, Property->GetFName(), ContainerIndex);
+
+			// Used to validate if we should skip the next property
+			UProperty* PossibleNextProperty = nullptr;
+
+			if ( Property->GetClass() == UArrayProperty::StaticClass() )
+			{
+				UArrayProperty* ArrayProperty  = static_cast<UArrayProperty*>( Property );
+				UProperty* ArrayTypeProperty = ArrayProperty->Inner;
+				DataprepPropertyChain.Emplace( ArrayTypeProperty, ArrayTypeProperty->GetFName(), INDEX_NONE );
+				PossibleNextProperty = ArrayTypeProperty;
+			}
+			else if ( Property->GetClass() == USetProperty::StaticClass() )
+			{
+				USetProperty* SetProperty = static_cast<USetProperty*>(Property);
+				UProperty* SetTypeProperty = SetProperty->ElementProp;
+				DataprepPropertyChain.Emplace(SetTypeProperty, SetTypeProperty->GetFName(), INDEX_NONE);
+				PossibleNextProperty = SetTypeProperty;
+			}
+
+			// We can't deal with map yet du to a lack of information from the property chain
+
+			
+			if ( PossibleNextProperty )
+			{ 
+				const TDoubleLinkedList<UProperty*>::TDoubleLinkedListNode* NextNode = CurrentNode->GetNextNode();
+				if ( NextNode && PossibleNextProperty == NextNode->GetValue() )
+				{
+					//skip the next property
+					CurrentNode = NextNode->GetNextNode();
+				}
+			}
 		}
 		else
 		{
@@ -204,13 +253,7 @@ bool FDataprepParameterizationUtils::IsPropertyChainValid(const TArray<FDataprep
 		return false;
 	}
 
-	// We will support those latter (there is still some bugs to fix before)
-	if ( PropertyChain.Num() > 1 )
-	{
-		return false;
-	}
-
-
+	bool bParentWasAContainer = false;
 	for ( const FDataprepPropertyLink& PropertyLink : PropertyChain )
 	{
 		UProperty* Property = PropertyLink.CachedProperty.Get();
@@ -219,12 +262,21 @@ bool FDataprepParameterizationUtils::IsPropertyChainValid(const TArray<FDataprep
 			return false;
 		}
 
-		// todo check for custom getter and setter
+		// Ensure that the properties are editable
+		if ( !bParentWasAContainer && !bool( Property->PropertyFlags & CPF_Edit ) )
+		{
+			return false;
+		}
 
-		// Temporary unsupported property check
+		if ( Property->ArrayDim <= PropertyLink.ContainerIndex )
 		{
-			// Reject container properties for now (there is still some bugs to fix before)
-			if ( DataprepParameterizationUtils::IsAContainerProperty( Property ) )
+			return false;
+		}
+
+		// Temporary unsupported properties check
+		{
+			// We don't support properties chain that are a subpart of a container property
+			if ( bParentWasAContainer )
 			{
 				return false;
 			}
@@ -234,6 +286,8 @@ bool FDataprepParameterizationUtils::IsPropertyChainValid(const TArray<FDataprep
 			{
 				return false;
 			}
+
+			bParentWasAContainer = DataprepParameterizationUtils::IsAContainerProperty( Property );
 		}
 	}
 
diff --git a/Engine/Plugins/Experimental/DataPrepEditor/Source/DataPrepCore/Public/DataPrepAsset.h b/Engine/Plugins/Experimental/DataPrepEditor/Source/DataPrepCore/Public/DataPrepAsset.h
--- a/Engine/Plugins/Experimental/DataPrepEditor/Source/DataPrepCore/Public/DataPrepAsset.h
+++ b/Engine/Plugins/Experimental/DataPrepEditor/Source/DataPrepCore/Public/DataPrepAsset.h
@@ -101,7 +101,7 @@ public:
 	 * Note on the objects param: The parameterized objects that should refresh their ui. If nullptr all widgets that can display some info on the parameterization should be refreshed
 	 */
 	DECLARE_EVENT_OneParam(UDataprepParameterization, FDataprepParameterizationStatusForObjectsChanged, const TSet<UObject*>* /** Objects */)
-	FDataprepParameterizationStatusForObjectsChanged OnParameterizationStatusForObjectsChanged;
+	FDataprepParameterizationStatusForObjectsChanged OnParameterizedObjectsChanged;
 
 	virtual UObject* GetParameterizationObject() override;
 
@@ -113,7 +113,7 @@ public:
 
 	void RemoveObjectPropertyFromParameterization(UDataprepParameterizableObject* Object, const TArray<struct FDataprepPropertyLink>& InPropertyChain);
 
-	void GetExistingParameterNamesForType(UProperty* Property, TSet<FString>& OutValidExistingNames, TSet<FString>& OutInvalidNames) const;
+	void GetExistingParameterNamesForType(UProperty* Property, bool bIsDescribingFullProperty, TSet<FString>& OutValidExistingNames, TSet<FString>& OutInvalidNames) const;
 
 	UDataprepParameterization* GetDataprepParameterization() { return Parameterization; }
 
diff --git a/Engine/Plugins/Experimental/DataPrepEditor/Source/DataPrepEditor/Private/Widgets/Action/SDataprepFloatFilter.cpp b/Engine/Plugins/Experimental/DataPrepEditor/Source/DataPrepEditor/Private/Widgets/Action/SDataprepFloatFilter.cpp
--- a/Engine/Plugins/Experimental/DataPrepEditor/Source/DataPrepEditor/Private/Widgets/Action/SDataprepFloatFilter.cpp
+++ b/Engine/Plugins/Experimental/DataPrepEditor/Source/DataPrepEditor/Private/Widgets/Action/SDataprepFloatFilter.cpp
@@ -84,7 +84,7 @@ void SDataprepFloatFilter::Construct(const FArguments& InArgs, UDataprepFloatFil
 			ToleranceParameterizationActionData = MakeShared<FDataprepParametrizationActionData>( *DataprepAsset, InFilter, PropertyChain );
 		}
 
-		OnParameterizationStatusForObjectsChangedHandle = DataprepAsset->OnParameterizationStatusForObjectsChanged.AddSP( this, &SDataprepFloatFilter::OnParameterizationStatusForObjectsChanged );
+		OnParameterizationStatusForObjectsChangedHandle = DataprepAsset->OnParameterizedObjectsChanged.AddSP( this, &SDataprepFloatFilter::OnParameterizationStatusForObjectsChanged );
 	}
 
 	UpdateVisualDisplay();
@@ -94,7 +94,7 @@ SDataprepFloatFilter::~SDataprepFloatFilter()
 {
 	if ( UDataprepAsset* DataprepAsset = FDataprepCoreUtils::GetDataprepAssetOfObject( Filter ) )
 	{
-		DataprepAsset->OnParameterizationStatusForObjectsChanged.Remove( OnParameterizationStatusForObjectsChangedHandle );
+		DataprepAsset->OnParameterizedObjectsChanged.Remove( OnParameterizationStatusForObjectsChangedHandle );
 	}
 }
 
diff --git a/Engine/Plugins/Experimental/DataPrepEditor/Source/DataPrepEditor/Private/Widgets/Action/SDataprepStringFilter.cpp b/Engine/Plugins/Experimental/DataPrepEditor/Source/DataPrepEditor/Private/Widgets/Action/SDataprepStringFilter.cpp
--- a/Engine/Plugins/Experimental/DataPrepEditor/Source/DataPrepEditor/Private/Widgets/Action/SDataprepStringFilter.cpp
+++ b/Engine/Plugins/Experimental/DataPrepEditor/Source/DataPrepEditor/Private/Widgets/Action/SDataprepStringFilter.cpp
@@ -59,7 +59,7 @@ void SDataprepStringFilter<FilterType>::Construct(const FArguments& InArgs, Filt
 			UserStringParameterizationActionData = MakeShared<FDataprepParametrizationActionData>( *DataprepAsset, InFilter, PropertyChain);
 		}
 
-		OnParameterizationStatusForObjectsChangedHandle = DataprepAsset->OnParameterizationStatusForObjectsChanged.AddSP( this, &SDataprepStringFilter<FilterType>::OnParameterizationStatusForObjectsChanged );
+		OnParameterizationStatusForObjectsChangedHandle = DataprepAsset->OnParameterizedObjectsChanged.AddSP( this, &SDataprepStringFilter<FilterType>::OnParameterizationStatusForObjectsChanged );
 	}
 
 	UpdateVisualDisplay();
@@ -70,7 +70,7 @@ SDataprepStringFilter<FilterType>::~SDataprepStringFilter()
 {
 	if ( UDataprepAsset* DataprepAsset = FDataprepCoreUtils::GetDataprepAssetOfObject( Filter ) )
 	{
-		DataprepAsset->OnParameterizationStatusForObjectsChanged.Remove( OnParameterizationStatusForObjectsChangedHandle );
+		DataprepAsset->OnParameterizedObjectsChanged.Remove( OnParameterizationStatusForObjectsChangedHandle );
 	}
 }
 
diff --git a/Engine/Plugins/Experimental/DataPrepEditor/Source/DataPrepEditor/Private/Widgets/DataprepAssetView.cpp b/Engine/Plugins/Experimental/DataPrepEditor/Source/DataPrepEditor/Private/Widgets/DataprepAssetView.cpp
--- a/Engine/Plugins/Experimental/DataPrepEditor/Source/DataPrepEditor/Private/Widgets/DataprepAssetView.cpp
+++ b/Engine/Plugins/Experimental/DataPrepEditor/Source/DataPrepEditor/Private/Widgets/DataprepAssetView.cpp
@@ -5,6 +5,7 @@
 #include "DataPrepAsset.h"
 #include "DataPrepEditor.h"
 #include "DataprepAssetInstance.h"
+#include "DataprepWidgets.h"
 
 #include "DataPrepContentConsumer.h"
 #include "DataprepEditorStyle.h"
@@ -34,6 +35,7 @@
 #include "Widgets/Text/SInlineEditableTextBlock.h"
 #include "Widgets/Text/STextBlock.h"
 
+
 #define LOCTEXT_NAMESPACE "DataprepAssetView"
 
 namespace DataprepEditorUtils
@@ -421,6 +423,8 @@ void SDataprepAssetView::Construct( const FArguments& InArgs, UDataprepAssetInte
 		];
 	}
 
+	TSharedPtr<SDataprepDetailsView> ParameterizationDetailsView;
+
 	ChildSlot
 	[
 		SNew(SBorder)
@@ -586,7 +590,7 @@ void SDataprepAssetView::Construct( const FArguments& InArgs, UDataprepAssetInte
 						[
 							SNew( DataprepWidgetUtils::SConstrainedBox )
 							[
-								SNew( SDataprepDetailsView )
+								SAssignNew( ParameterizationDetailsView, SDataprepDetailsView )
 								.Object( DataprepAssetInterfacePtr->GetParameterizationObject() )
 								.ColumnSizeData( ColumnSizeData )
 								.Spacing( 10.0f )
@@ -608,13 +612,35 @@ void SDataprepAssetView::Construct( const FArguments& InArgs, UDataprepAssetInte
 			]
 		]
 	];
+
+	if ( DataprepAssetInterfacePtr->IsA<UDataprepAsset>() )
+	{
+		TWeakObjectPtr<UDataprepAsset> DataprepAsset = static_cast<UDataprepAsset*>( DataprepAssetInterfacePtr.Get() );
+		OnParameterizationWasEdited = DataprepAsset->OnParameterizedObjectsChanged.AddLambda( [ParameterizationDetailsView, DataprepAsset](const TSet<UObject*>* Objects)
+		{
+			if ( Objects && Objects->Contains( DataprepAsset->GetParameterizationObject() ) )
+			{
+				ParameterizationDetailsView->ForceRefresh();
+			}
+		});
+	}
+
 }
 
 SDataprepAssetView::~SDataprepAssetView()
 {
-	if( UDataprepAssetInterface* DataprepAsset = DataprepAssetInterfacePtr.Get() )
+	if( UDataprepAssetInterface* DataprepAssetInterface = DataprepAssetInterfacePtr.Get() )
 	{
-		DataprepAsset->GetOnChanged().RemoveAll( this );
+		DataprepAssetInterface->GetOnChanged().RemoveAll( this );
+
+		if ( OnParameterizationWasEdited.IsValid() )
+		{
+			if ( DataprepAssetInterfacePtr->IsA<UDataprepAsset>() )
+			{
+				TWeakObjectPtr<UDataprepAsset> DataprepAsset = static_cast<UDataprepAsset*>( DataprepAssetInterfacePtr.Get() );
+				DataprepAsset->OnParameterizedObjectsChanged.Remove( OnParameterizationWasEdited );
+			}
+		}
 	}
 }
 
diff --git a/Engine/Plugins/Experimental/DataPrepEditor/Source/DataPrepEditor/Private/Widgets/DataprepAssetView.h b/Engine/Plugins/Experimental/DataPrepEditor/Source/DataPrepEditor/Private/Widgets/DataprepAssetView.h
--- a/Engine/Plugins/Experimental/DataPrepEditor/Source/DataPrepEditor/Private/Widgets/DataprepAssetView.h
+++ b/Engine/Plugins/Experimental/DataPrepEditor/Source/DataPrepEditor/Private/Widgets/DataprepAssetView.h
@@ -10,6 +10,7 @@
 #include "DataPrepEditor.h"
 #include "DataprepWidgets.h"
 
+#include "Delegates/IDelegateInstance.h"
 #include "Widgets/Input/STextComboBox.h"
 #include "Widgets/SCompoundWidget.h"
 #include "Widgets/Text/STextBlock.h"
@@ -57,6 +58,8 @@ private:
 
 	/** Relative width to control splitters */
 	float ColumnWidth;
+
+	FDelegateHandle OnParameterizationWasEdited;
 };
 
 // Inspired from SKismetInspector class
diff --git a/Engine/Plugins/Experimental/DataPrepEditor/Source/DataPrepEditor/Private/Widgets/DataprepWidgets.cpp b/Engine/Plugins/Experimental/DataPrepEditor/Source/DataPrepEditor/Private/Widgets/DataprepWidgets.cpp
--- a/Engine/Plugins/Experimental/DataPrepEditor/Source/DataPrepEditor/Private/Widgets/DataprepWidgets.cpp
+++ b/Engine/Plugins/Experimental/DataPrepEditor/Source/DataPrepEditor/Private/Widgets/DataprepWidgets.cpp
@@ -575,7 +575,9 @@ void SDataprepDetailsView::OnObjectReplaced(const TMap<UObject*, UObject*>& Repl
 
 void SDataprepDetailsView::ForceRefresh()
 {
-	bRefreshObjectToDisplay = true;
+	// Hotfix for 4.24 (Remove the ui flickering)
+	Construct();
+	//bRefreshObjectToDisplay = true;
 }
 
 void SDataprepDetailsView::OnDataprepParameterizationStatusForObjectsChanged(const TSet<UObject*>* Objects)
@@ -798,7 +800,7 @@ void SDataprepDetailsView::Construct()
 		UDataprepAsset* DataprepAsset = FDataprepParameterizationUtils::GetDataprepAssetForParameterization( DetailedObject );
 		if ( DataprepAsset )
 		{
-			OnDataprepParameterizationStatusForObjectsChangedHandle = DataprepAsset->OnParameterizationStatusForObjectsChanged.AddSP( this, &SDataprepDetailsView::OnDataprepParameterizationStatusForObjectsChanged );
+			OnDataprepParameterizationStatusForObjectsChangedHandle = DataprepAsset->OnParameterizedObjectsChanged.AddSP( this, &SDataprepDetailsView::OnDataprepParameterizationStatusForObjectsChanged );
 		}
 
 		if ( DetailedObject->IsA<UDataprepParameterizableObject>() )
@@ -864,7 +866,7 @@ SDataprepDetailsView::~SDataprepDetailsView()
 
 	if ( UDataprepAsset* DataprepAsset = DataprepAssetForParameterization.Get() )
 	{
-		DataprepAsset->OnParameterizationStatusForObjectsChanged.Remove( OnDataprepParameterizationStatusForObjectsChangedHandle );
+		DataprepAsset->OnParameterizedObjectsChanged.Remove( OnDataprepParameterizationStatusForObjectsChangedHandle );
 	}
 }
 
diff --git a/Engine/Plugins/Experimental/DataPrepEditor/Source/DataPrepEditor/Private/Widgets/Parameterization/SDataprepLinkToParameter.cpp b/Engine/Plugins/Experimental/DataPrepEditor/Source/DataPrepEditor/Private/Widgets/Parameterization/SDataprepLinkToParameter.cpp
--- a/Engine/Plugins/Experimental/DataPrepEditor/Source/DataPrepEditor/Private/Widgets/Parameterization/SDataprepLinkToParameter.cpp
+++ b/Engine/Plugins/Experimental/DataPrepEditor/Source/DataPrepEditor/Private/Widgets/Parameterization/SDataprepLinkToParameter.cpp
@@ -34,7 +34,8 @@ void SDataprepLinkToParameter::Construct(const FArguments& InArgs, const TShared
 	{
 		if ( UProperty* Property = ParameterizationActionData->PropertyChain.Last().CachedProperty.Get() )
 		{
-			ParameterizationActionData->DataprepAsset->GetExistingParameterNamesForType( Property, ValidExistingNames, InvalidNames );
+			bool bIsDescribingFullProperty = ParameterizationActionData->PropertyChain.Last().ContainerIndex == INDEX_NONE;
+			ParameterizationActionData->DataprepAsset->GetExistingParameterNamesForType( Property, bIsDescribingFullProperty , ValidExistingNames, InvalidNames );
 
 			FString SuggestedName = ParameterizationActionData->PropertyChain.Last().PropertyName.ToString();
 
