Added some missing EditAnywhere to uproperties of the DataprepFilters

diff --git a/Engine/Plugins/Experimental/DataPrepEditor/Source/DataPrepCore/Public/SelectionSystem/DataprepFloatFilter.h b/Engine/Plugins/Experimental/DataPrepEditor/Source/DataPrepCore/Public/SelectionSystem/DataprepFloatFilter.h
--- a/Engine/Plugins/Experimental/DataPrepEditor/Source/DataPrepCore/Public/SelectionSystem/DataprepFloatFilter.h
+++ b/Engine/Plugins/Experimental/DataPrepEditor/Source/DataPrepCore/Public/SelectionSystem/DataprepFloatFilter.h
@@ -51,14 +51,14 @@ private:
 	UDataprepFloatFetcher* FloatFetcher;
 
 	// The criteria selected by the user
-	UPROPERTY()
+	UPROPERTY(EditAnywhere, Category = Filter)
 	EDataprepFloatMatchType FloatMatchingCriteria;
 
 	// The value to use for the equality check
-	UPROPERTY()
+	UPROPERTY(EditAnywhere, Category = Filter)
 	float EqualValue;
 
 	// The value used for the tolerance when doing a nearly equal
-	UPROPERTY()
+	UPROPERTY(EditAnywhere, Category = Filter)
 	float Tolerance = KINDA_SMALL_NUMBER;
 };
diff --git a/Engine/Plugins/Experimental/DataPrepEditor/Source/DataPrepCore/Public/SelectionSystem/DataprepStringFilter.h b/Engine/Plugins/Experimental/DataPrepEditor/Source/DataPrepCore/Public/SelectionSystem/DataprepStringFilter.h
--- a/Engine/Plugins/Experimental/DataPrepEditor/Source/DataPrepCore/Public/SelectionSystem/DataprepStringFilter.h
+++ b/Engine/Plugins/Experimental/DataPrepEditor/Source/DataPrepCore/Public/SelectionSystem/DataprepStringFilter.h
@@ -46,11 +46,11 @@ public:
 
 private:
 	// The criteria selected by the user
-	UPROPERTY()
+	UPROPERTY(EditAnywhere, Category = Filter)
 	EDataprepStringMatchType StringMatchingCriteria;
 
 	// The string entered by the user
-	UPROPERTY()
+	UPROPERTY(EditAnywhere, Category = Filter)
 	FString UserString;
 
 	// The source of string selected by the user
diff --git a/Engine/Plugins/Experimental/DataPrepEditor/Source/DataPrepCore/Public/SelectionSystem/DataprepStringsArrayFilter.h b/Engine/Plugins/Experimental/DataPrepEditor/Source/DataPrepCore/Public/SelectionSystem/DataprepStringsArrayFilter.h
--- a/Engine/Plugins/Experimental/DataPrepEditor/Source/DataPrepCore/Public/SelectionSystem/DataprepStringsArrayFilter.h
+++ b/Engine/Plugins/Experimental/DataPrepEditor/Source/DataPrepCore/Public/SelectionSystem/DataprepStringsArrayFilter.h
@@ -33,11 +33,11 @@ public:
 
 private:
 	// The criteria selected by the user
-	UPROPERTY()
+	UPROPERTY(EditAnywhere, Category = Filter)
 	EDataprepStringMatchType StringMatchingCriteria;
 
 	// The string entered by the user
-	UPROPERTY()
+	UPROPERTY(EditAnywhere, Category = Filter)
 	FString UserString;
 
 	// The source of string selected by the user
