Fixed a crash when force refreshing a dataprep details view.

#rnx

diff --git a/Engine/Plugins/Experimental/DataPrepEditor/Source/DataPrepEditor/Private/Widgets/DataprepWidgets.cpp b/Engine/Plugins/Experimental/DataPrepEditor/Source/DataPrepEditor/Private/Widgets/DataprepWidgets.cpp
--- a/Engine/Plugins/Experimental/DataPrepEditor/Source/DataPrepEditor/Private/Widgets/DataprepWidgets.cpp
+++ b/Engine/Plugins/Experimental/DataPrepEditor/Source/DataPrepEditor/Private/Widgets/DataprepWidgets.cpp
@@ -576,6 +576,9 @@ void SDataprepDetailsView::OnObjectReplaced(const TMap<UObject*, UObject*>& Repl
 void SDataprepDetailsView::ForceRefresh()
 {
 	// Hotfix for 4.24 (Remove the ui flickering)
+	TArray< UObject* > Objects;
+	Objects.Add( DetailedObject );
+	Generator->SetObjects( Objects );
 	Construct();
 	//bRefreshObjectToDisplay = true;
 }
