Fix include capitalization
#rnx

diff --git a/Engine/Plugins/Experimental/DataPrepEditor/Source/DataPrepCore/Private/Parameterization/DataprepParameterization.cpp b/Engine/Plugins/Experimental/DataPrepEditor/Source/DataPrepCore/Private/Parameterization/DataprepParameterization.cpp
--- a/Engine/Plugins/Experimental/DataPrepEditor/Source/DataPrepCore/Private/Parameterization/DataprepParameterization.cpp
+++ b/Engine/Plugins/Experimental/DataPrepEditor/Source/DataPrepCore/Private/Parameterization/DataprepParameterization.cpp
@@ -3,7 +3,7 @@
 #include "Parameterization/DataprepParameterization.h"
 
 #include "DataPrepAsset.h"
-#include "DataPrepAssetInstance.h"
+#include "DataprepAssetInstance.h"
 #include "DataprepParameterizableObject.h"
 #include "DataprepParameterizationArchive.h"
 #include "Parameterization/DataprepParameterizationUtils.h"
