Fixed a duplication issue with with the dataprep asset. The post load of the dataprep parameterization wasn't generating a proper container.

#rnx

diff --git a/Engine/Plugins/Experimental/DataPrepEditor/Source/DataPrepCore/Private/Parameterization/DataprepParameterization.cpp b/Engine/Plugins/Experimental/DataPrepEditor/Source/DataPrepCore/Private/Parameterization/DataprepParameterization.cpp
--- a/Engine/Plugins/Experimental/DataPrepEditor/Source/DataPrepCore/Private/Parameterization/DataprepParameterization.cpp
+++ b/Engine/Plugins/Experimental/DataPrepEditor/Source/DataPrepCore/Private/Parameterization/DataprepParameterization.cpp
@@ -888,7 +888,18 @@ void UDataprepParameterization::PostLoad()
 			BindingsContainer = NewObject<UDataprepParameterizationBindings>(this, NAME_None, RF_Transactional);
 		}
 
+		PrepareCustomClassForNewClassGeneration();
+		UClass* OldClass = CustomContainerClass;
+		CustomContainerClass = nullptr;
+
 		LoadParameterization();
+
+		if ( OldClass )
+		{
+			DoReinstancing( OldClass, false );
+			OnTellInstancesToReloadTheirSerializedData.Broadcast();
+			CastChecked<UDataprepAsset>( GetOuter() )->OnParameterizedObjectsChanged.Broadcast( nullptr );
+		}
 	}
 	Super::PostLoad();
 }
