Fixed a crash  that occured afther a force refreshing a dataprep details view.

#rnx

diff --git a/Engine/Plugins/Experimental/DataPrepEditor/Source/DataPrepEditor/Private/Widgets/DataprepWidgets.cpp b/Engine/Plugins/Experimental/DataPrepEditor/Source/DataPrepEditor/Private/Widgets/DataprepWidgets.cpp
--- a/Engine/Plugins/Experimental/DataPrepEditor/Source/DataPrepEditor/Private/Widgets/DataprepWidgets.cpp
+++ b/Engine/Plugins/Experimental/DataPrepEditor/Source/DataPrepEditor/Private/Widgets/DataprepWidgets.cpp
@@ -575,12 +575,9 @@ void SDataprepDetailsView::OnObjectReplaced(const TMap<UObject*, UObject*>& Repl
 
 void SDataprepDetailsView::ForceRefresh()
 {
-	// Hotfix for 4.24 (Remove the ui flickering)
-	TArray< UObject* > Objects;
-	Objects.Add( DetailedObject );
-	Generator->SetObjects( Objects );
-	Construct();
-	//bRefreshObjectToDisplay = true;
+	// ueent_hotfix Hotfix for 4.24 (Remove the ui flickering)
+	InvalidatePrepass();
+	bRefreshObjectToDisplay = true;
 }
 
 void SDataprepDetailsView::OnDataprepParameterizationStatusForObjectsChanged(const TSet<UObject*>* Objects)
@@ -797,6 +794,7 @@ void SDataprepDetailsView::Construct(const FArguments& InArgs)
 void SDataprepDetailsView::Construct()
 {
 	DataprepAssetForParameterization.Reset();
+	bHasCustomPrepass = true;
 
 	if ( DetailedObject )
 	{
@@ -873,21 +871,6 @@ SDataprepDetailsView::~SDataprepDetailsView()
 	}
 }
 
-void SDataprepDetailsView::Tick(const FGeometry& AllottedGeometry, const double InCurrentTime, const float InDeltaTime)
-{
-	Super::Tick( AllottedGeometry, InCurrentTime, InDeltaTime );
-
-	if ( bRefreshObjectToDisplay )
-	{
-		TArray< UObject* > Objects;
-		Objects.Add( DetailedObject );
-		Generator->SetObjects( Objects );
-		 
-		Construct();
-		bRefreshObjectToDisplay = false;
-	}
-}
-
 void SDataprepDetailsView::SetObjectToDisplay(UObject& Object)
 {
 	UObject* NewObjectToDisplay = &Object;
@@ -909,6 +892,19 @@ void SDataprepDetailsView::AddReferencedObjects(FReferenceCollector& Collector)
 	Collector.AddReferencedObjects( TrackedProperties );
 }
 
+bool SDataprepDetailsView::CustomPrepass(float LayoutScaleMultiplier)
+{
+	if ( bRefreshObjectToDisplay )
+	{
+		TArray< UObject* > Objects;
+		Objects.Add( DetailedObject );
+		Generator->SetObjects( Objects );
+		Construct();
+		bRefreshObjectToDisplay = false;
+	}
+
+	return true;
+}
 
 void SDataprepContextMenuOverride::Construct(const FArguments& InArgs)
 {
diff --git a/Engine/Plugins/Experimental/DataPrepEditor/Source/DataPrepEditor/Private/Widgets/DataprepWidgets.h b/Engine/Plugins/Experimental/DataPrepEditor/Source/DataPrepEditor/Private/Widgets/DataprepWidgets.h
--- a/Engine/Plugins/Experimental/DataPrepEditor/Source/DataPrepEditor/Private/Widgets/DataprepWidgets.h
+++ b/Engine/Plugins/Experimental/DataPrepEditor/Source/DataPrepEditor/Private/Widgets/DataprepWidgets.h
@@ -138,8 +138,6 @@ public:
 
 	~SDataprepDetailsView();
 
-	virtual void Tick(const FGeometry& AllottedGeometry, const double InCurrentTime, const float InDeltaTime) override;	
-
 	void SetObjectToDisplay(UObject& Object);
 
 	void ForceRefresh();
@@ -149,6 +147,9 @@ public:
 	// End of FGCObject interface
 
 protected:
+	// ueent_hotfix Hack for 4.24 allow to refresh the ui in between two frame without any flickering
+	virtual bool CustomPrepass(float LayoutScaleMultiplier) override;
+
 	// Begin SWidget overrides.
 	// #ueent_todo: This is temporary until we find a better solution to the splitter issue
 	//				See SConstrainedBox's trick in cpp file
