Fixed a crash when pressing enter to link a dataprep property to the parameterization.

#rnx

diff --git a/Engine/Plugins/Experimental/DataPrepEditor/Source/DataPrepEditor/Private/Widgets/Parameterization/SDataprepLinkToParameter.cpp b/Engine/Plugins/Experimental/DataPrepEditor/Source/DataPrepEditor/Private/Widgets/Parameterization/SDataprepLinkToParameter.cpp
--- a/Engine/Plugins/Experimental/DataPrepEditor/Source/DataPrepEditor/Private/Widgets/Parameterization/SDataprepLinkToParameter.cpp
+++ b/Engine/Plugins/Experimental/DataPrepEditor/Source/DataPrepEditor/Private/Widgets/Parameterization/SDataprepLinkToParameter.cpp
@@ -179,9 +179,8 @@ void SDataprepLinkToParameter::OnTextChanged(const FText& Text)
 
 void SDataprepLinkToParameter::OnTextCommited(const FText& Text, ETextCommit::Type CommitType)
 {
-
 	ParameterName = Text.ToString();
-	bool bHasError = SetErrorMessage();
+	bool bHasError = ParameterName.IsEmpty() || InvalidNames.Contains( ParameterName );
 	
 	if ( CommitType == ETextCommit::OnEnter )
 	{
