Finalized implementation of Create Proxy Mesh

diff --git a/Engine/Plugins/Experimental/DataPrepEditor/Source/DataPrepEditor/Private/DataPrepEditor.cpp b/Engine/Plugins/Experimental/DataPrepEditor/Source/DataPrepEditor/Private/DataPrepEditor.cpp
--- a/Engine/Plugins/Experimental/DataPrepEditor/Source/DataPrepEditor/Private/DataPrepEditor.cpp
+++ b/Engine/Plugins/Experimental/DataPrepEditor/Source/DataPrepEditor/Private/DataPrepEditor.cpp
@@ -413,7 +413,7 @@ void FDataprepEditor::InitDataprepEditor(const EToolkitMode::Type Mode, const TS
 
 	ActionsContext = MakeShareable( new FDataprepActionContext() );
 
-	ActionsContext->SetTransientContentFolder( GetTransientContentFolder() / TEXT("Pipeline") )
+	ActionsContext->SetTransientContentFolder( GetTransientContentFolder() / DataprepAssetInterfacePtr->GetName() / TEXT("Pipeline") )
 		.SetLogger( TSharedPtr<IDataprepLogger>( new FDataprepCoreUtils::FDataprepLogger ) )
 		.SetCanExecuteNextStep( CanExecuteNextStepFunc )
 		.SetActionsContextChanged( ActionsContextChangedFunc );
diff --git a/Engine/Plugins/Experimental/DataPrepEditor/Source/DataPrepLibraries/Private/DataprepEditingOperations.cpp b/Engine/Plugins/Experimental/DataPrepEditor/Source/DataPrepLibraries/Private/DataprepEditingOperations.cpp
--- a/Engine/Plugins/Experimental/DataPrepEditor/Source/DataPrepLibraries/Private/DataprepEditingOperations.cpp
+++ b/Engine/Plugins/Experimental/DataPrepEditor/Source/DataPrepLibraries/Private/DataprepEditingOperations.cpp
@@ -24,6 +24,8 @@
 #include "LevelSequence.h"
 #include "Materials/MaterialInstance.h"
 #include "MeshMergeModule.h"
+#include "ObjectTools.h"
+#include "StaticMeshAttributes.h"
 #include "StaticMeshResources.h"
 
 #define LOCTEXT_NAMESPACE "DatasmithEditingOperations"
@@ -163,8 +165,8 @@ void UDataprepMergeActorsOperation::OnExecution_Implementation(const FDataprepCo
 
 	DatasmithEditingOperationsUtils::GetActorsToMerge(CurrentWorld, InContext.Objects, ActorsToMerge, ComponentsToMerge);
 
-	// Nothing to do if there is none or only one static mesh actor
-	if( ActorsToMerge.Num() < 2 && ComponentsToMerge.Num() < 2)
+	// Nothing to do if there is only one component to merge
+	if( ComponentsToMerge.Num() < 2)
 	{
 		UE_LOG( LogDataprep, Log, TEXT("No static mesh actors to merge") );
 		return;
@@ -295,10 +297,10 @@ void UDataprepCreateProxyMeshOperation::OnExecution_Implementation(const FDatapr
 
 	DatasmithEditingOperationsUtils::GetActorsToMerge(CurrentWorld, InContext.Objects, ActorsToMerge, ComponentsToMerge);
 
-	// Nothing to do if there is none or only one static mesh actor
-	if(ComponentsToMerge.Num() < 2)
+	// Nothing to do if there is no static mesh components to merge
+	if(ComponentsToMerge.Num() == 0)
 	{
-		UE_LOG(LogDataprep, Log, TEXT("No static mesh actors to merge"));
+		UE_LOG(LogDataprep, Log, TEXT("No static mesh to merge"));
 		return;
 	}
 
@@ -317,8 +319,58 @@ void UDataprepCreateProxyMeshOperation::OnExecution_Implementation(const FDatapr
 
 	DataprepOperationsLibraryUtil::FStaticMeshBuilder StaticMeshBuilder( StaticMeshes );
 
+	// Update the settings for geometry
 	FMeshProxySettings ProxySettings;
-	ProxySettings.ScreenSize = ProxySettings.ScreenSize * ((100.0f - ReductionPercent) / 100.f);
+	ProxySettings.bOverrideVoxelSize = false;
+	
+	const float Coefficient = 2.0f * Quality / 100.0f;
+
+	const float MinScreenSize = Coefficient <= 1.0f ? 100.0f : 300.0f;
+	const float MaxScreenSize = Coefficient <= 1.0f ? 300.0f : 1200.0f;
+	ProxySettings.ScreenSize = FMath::FloorToInt( FMath::Lerp( MinScreenSize, MaxScreenSize, Coefficient <= 1.0f ? Coefficient : Coefficient - 1.0f ) + 0.5f );
+
+	// Determine if incoming lightmap UVs are usable
+	ProxySettings.bReuseMeshLightmapUVs = true;
+	StaticMeshes.Empty( ComponentsToMerge.Num() );
+	for(UPrimitiveComponent* PrimitiveComponent : ComponentsToMerge)
+	{
+		if(UStaticMeshComponent* StaticMeshComponent = Cast<UStaticMeshComponent>(PrimitiveComponent))
+		{
+			if(StaticMeshComponent->GetStaticMesh())
+			{
+				StaticMeshes.Add( StaticMeshComponent->GetStaticMesh() );
+			}
+		}
+	}
+
+	for(UStaticMesh* StaticMesh : StaticMeshes)
+	{
+		const FMeshBuildSettings& BuildSettings = StaticMesh->GetSourceModel(0).BuildSettings;
+		if(!BuildSettings.bGenerateLightmapUVs)
+		{
+			ProxySettings.bReuseMeshLightmapUVs = false;
+			break;
+		}
+		else if(FMeshDescription* MeshDescription = StaticMesh->GetMeshDescription( 0 ))
+		{
+			FStaticMeshAttributes Attributes(*MeshDescription);
+			bool bHasValidLightmapUVs = Attributes.GetVertexInstanceUVs().IsValid() &&
+				Attributes.GetVertexInstanceUVs().GetNumIndices() > BuildSettings.SrcLightmapIndex &&
+				Attributes.GetVertexInstanceUVs().GetNumIndices() > BuildSettings.DstLightmapIndex;
+			if(!bHasValidLightmapUVs)
+			{
+				ProxySettings.bReuseMeshLightmapUVs = false;
+				break;
+			}
+		}
+	}
+
+	// Update the settings for materials
+	ProxySettings.MaterialSettings.bMetallicMap = true;
+	ProxySettings.MaterialSettings.bRoughnessMap = true;
+
+	const int32 TextureSize = (Coefficient <= 0.5f) ? 512 : ((Coefficient <= 1.0f) ? 1024 : ((Coefficient <= 1.5f) ? 2048 : 4096 ));
+	ProxySettings.MaterialSettings.TextureSize = FIntPoint( TextureSize, TextureSize );
 
 	// Generate proxy mesh and proxy material assets 
 	FCreateProxyDelegate ProxyDelegate;
@@ -350,6 +402,28 @@ void UDataprepCreateProxyMeshOperation::OnExecution_Implementation(const FDatapr
 		MergedActor->GetStaticMeshComponent()->SetStaticMesh(MergedMesh);
 		MergedActor->SetActorLabel(NewActorLabel.IsEmpty() ? TEXT("Proxy_Actor") : *NewActorLabel);
 		CurrentWorld->UpdateCullDistanceVolumes(MergedActor, MergedActor->GetStaticMeshComponent());
+
+		// Add the other assets created by the merge, i.e. material, texture, etc, to the context
+		TArray< TPair< UObject*, UObject* > > RedirectionMap;
+		RedirectionMap.Reserve( AssetsToSync.Num() );
+
+		for(UObject* Object : AssetsToSync)
+		{
+			if( Cast<UStaticMesh>(Object) != ProxyMesh)
+			{
+				const FString DefaultName = Object->GetClass()->GetName() + TEXT("_Proxy");
+				UObject* AssetFromMerge = AddAsset(Object, NewActorLabel.IsEmpty() ? *DefaultName : *NewActorLabel);
+
+				RedirectionMap.Emplace( AssetFromMerge, Object );
+			}
+		}
+
+		// Update references accordingly
+		for(TPair< UObject*, UObject* >& MapEntry : RedirectionMap)
+		{
+			TArray<UObject*> ObjectsToReplace(&MapEntry.Value, 1);
+			ObjectTools::ForceReplaceReferences( MapEntry.Key, ObjectsToReplace );
+		}
 	});
 
 	FGuid JobGuid = FGuid::NewGuid();
@@ -698,7 +772,7 @@ namespace DatasmithEditingOperationsUtils
 					bool bMeshActorIsValid = false;
 					for(UStaticMeshComponent* MeshComponent : ComponentArray)
 					{
-						if(MeshComponent->GetStaticMesh())
+						if(MeshComponent->GetStaticMesh() && MeshComponent->GetStaticMesh()->GetSourceModels().Num() > 0)
 						{
 							bMeshActorIsValid = true;
 							ComponentsToMerge.Add(MeshComponent);
diff --git a/Engine/Plugins/Experimental/DataPrepEditor/Source/DataPrepLibraries/Private/DataprepEditingOperations.h b/Engine/Plugins/Experimental/DataPrepEditor/Source/DataPrepLibraries/Private/DataprepEditingOperations.h
--- a/Engine/Plugins/Experimental/DataPrepEditor/Source/DataPrepLibraries/Private/DataprepEditingOperations.h
+++ b/Engine/Plugins/Experimental/DataPrepEditor/Source/DataPrepLibraries/Private/DataprepEditingOperations.h
@@ -34,7 +34,7 @@ protected:
 	//~ End UDataprepOperation Interface
 };
 
-UCLASS(Experimental, Category = ObjectOperation, Meta = (DisplayName="Merge Actors", ToolTip = "Merge the meshes into a unique mesh with the provided StaticMeshActors") )
+UCLASS(Experimental, Category = ObjectOperation, Meta = (DisplayName="Merge", ToolTip = "Collect geometry from selected actors and merge them into single mesh.") )
 class UDataprepMergeActorsOperation : public UDataprepEditingOperation
 {
 	GENERATED_BODY()
@@ -53,7 +53,7 @@ public:
 public:
 	virtual FText GetCategory_Implementation() const override
 	{
-		return FDataprepOperationCategories::ObjectOperation;
+		return FDataprepOperationCategories::ActorOperation;
 	}
 
 	/** Settings to use for the merge operation */
@@ -87,14 +87,14 @@ protected:
 	AStaticMeshActor* MergedActor;
 };
 
-UCLASS(Experimental, Category = ObjectOperation, Meta = (DisplayName = "Create Proxy Mesh", ToolTip = "Merge the meshes into a (simplified) proxy mesh"))
+UCLASS(Experimental, Category = ObjectOperation, Meta = (DisplayName = "Create Proxy Mesh", ToolTip = "Collect geometry from selected actors and merge them into single mesh with reduction."))
 class UDataprepCreateProxyMeshOperation : public UDataprepEditingOperation
 {
 	GENERATED_BODY()
 
 public:
 	UDataprepCreateProxyMeshOperation()
-		: ReductionPercent(50.f)
+		: Quality(50.f)
 		, MergedMesh(nullptr)
 		, MergedActor(nullptr)
 	{
@@ -104,7 +104,7 @@ public:
 	//~ Begin UDataprepOperation Interface
 	virtual FText GetCategory_Implementation() const override
 	{
-		return FDataprepOperationCategories::ObjectOperation;
+		return FDataprepOperationCategories::ActorOperation;
 	}
 
 	/** Settings to use for the create proxy operation */
@@ -112,7 +112,7 @@ public:
 	FString NewActorLabel;
 
 	UPROPERTY(EditAnywhere, Category = ProxySettings,  meta = (UIMin = 0, UIMax = 100))
-	float ReductionPercent;
+	float Quality;
 
 protected:
 	virtual void OnExecution_Implementation(const FDataprepContext& InContext) override;
