Fixed a crash while loading a dataprep instance

#rnx

diff --git a/Engine/Plugins/Experimental/DataPrepEditor/Source/DataPrepCore/Private/Parameterization/DataprepParameterization.cpp b/Engine/Plugins/Experimental/DataPrepEditor/Source/DataPrepCore/Private/Parameterization/DataprepParameterization.cpp
--- a/Engine/Plugins/Experimental/DataPrepEditor/Source/DataPrepCore/Private/Parameterization/DataprepParameterization.cpp
+++ b/Engine/Plugins/Experimental/DataPrepEditor/Source/DataPrepCore/Private/Parameterization/DataprepParameterization.cpp
@@ -888,6 +888,7 @@ void UDataprepParameterization::PostLoad()
 			BindingsContainer = NewObject<UDataprepParameterizationBindings>(this, NAME_None, RF_Transactional);
 		}
 
+		// ueent_hotfix revisit this code has renaming a object while a linker is active is dangerous (this was put here so that the duplicate object would work properly)
 		PrepareCustomClassForNewClassGeneration();
 		UClass* OldClass = CustomContainerClass;
 		CustomContainerClass = nullptr;
@@ -1292,7 +1293,7 @@ void UDataprepParameterization::PrepareCustomClassForNewClassGeneration()
 		CustomContainerClass->ClassFlags |= CLASS_NewerVersionExists;
 		CustomContainerClass->SetFlags( RF_NewerVersionExists );
 		CustomContainerClass->ClearFlags( RF_Public | RF_Standalone );
-		CustomContainerClass->Rename( *OldClassName, nullptr, REN_DontCreateRedirectors | REN_DoNotDirty);
+		CustomContainerClass->Rename( *OldClassName, GetTransientPackage(), REN_DontCreateRedirectors | REN_DoNotDirty);
 	}
 }
 
@@ -1555,7 +1556,7 @@ void UDataprepParameterizationInstance::LoadParameterization()
 
 	if ( !SourceParameterization->CustomContainerClass )
 	{
-		SourceParameterization->LoadParameterization();
+		SourceParameterization->ConditionalPostLoad();
 	}
 
 	if ( !ParameterizationInstance )
