Reworked the naming of the assets generated by the 'Create Proxy Mesh' operation

diff --git a/Engine/Plugins/Experimental/DataPrepEditor/Source/DataPrepLibraries/Private/DataprepEditingOperations.cpp b/Engine/Plugins/Experimental/DataPrepEditor/Source/DataPrepLibraries/Private/DataprepEditingOperations.cpp
--- a/Engine/Plugins/Experimental/DataPrepEditor/Source/DataPrepLibraries/Private/DataprepEditingOperations.cpp
+++ b/Engine/Plugins/Experimental/DataPrepEditor/Source/DataPrepLibraries/Private/DataprepEditingOperations.cpp
@@ -372,6 +372,8 @@ void UDataprepCreateProxyMeshOperation::OnExecution_Implementation(const FDatapr
 	const int32 TextureSize = (Coefficient <= 0.5f) ? 512 : ((Coefficient <= 1.0f) ? 1024 : ((Coefficient <= 1.5f) ? 2048 : 4096 ));
 	ProxySettings.MaterialSettings.TextureSize = FIntPoint( TextureSize, TextureSize );
 
+	const TCHAR* ProxyBasePackageName = TEXT("TOREPLACE");
+
 	// Generate proxy mesh and proxy material assets 
 	FCreateProxyDelegate ProxyDelegate;
 	ProxyDelegate.BindLambda( [&](const FGuid Guid, TArray<UObject*>& AssetsToSync )
@@ -411,8 +413,8 @@ void UDataprepCreateProxyMeshOperation::OnExecution_Implementation(const FDatapr
 		{
 			if( Cast<UStaticMesh>(Object) != ProxyMesh)
 			{
-				const FString DefaultName = Object->GetClass()->GetName() + TEXT("_Proxy");
-				UObject* AssetFromMerge = AddAsset(Object, NewActorLabel.IsEmpty() ? *DefaultName : *NewActorLabel);
+				const FString AssetName = Object->GetName().Replace( ProxyBasePackageName, NewActorLabel.IsEmpty() ? *GetDisplayOperationName().ToString() : *NewActorLabel, ESearchCase::CaseSensitive );
+				UObject* AssetFromMerge = AddAsset( Object, *AssetName );
 
 				RedirectionMap.Emplace( AssetFromMerge, Object );
 			}
@@ -429,7 +431,7 @@ void UDataprepCreateProxyMeshOperation::OnExecution_Implementation(const FDatapr
 	FGuid JobGuid = FGuid::NewGuid();
 
 	const IMeshMergeUtilities& MergeUtilities = FModuleManager::Get().LoadModuleChecked<IMeshMergeModule>("MeshMergeUtilities").GetUtilities();
-	MergeUtilities.CreateProxyMesh(ActorsToMerge, ProxySettings, nullptr, GetTransientPackage(), FString(), JobGuid, ProxyDelegate);
+	MergeUtilities.CreateProxyMesh(ActorsToMerge, ProxySettings, nullptr, GetTransientPackage(), ProxyBasePackageName, JobGuid, ProxyDelegate);
 
 	// Position the merged actor at the right location
 	if(MergedActor->GetRootComponent() == nullptr)
