Changed the category and the tooltip of the set mobility operator from the Dataprep Editor plugin.

#rnx

diff --git a/Engine/Plugins/Experimental/DataPrepEditor/Source/DataPrepLibraries/Private/DataprepOperations.h b/Engine/Plugins/Experimental/DataPrepEditor/Source/DataPrepLibraries/Private/DataprepOperations.h
--- a/Engine/Plugins/Experimental/DataPrepEditor/Source/DataPrepLibraries/Private/DataprepOperations.h
+++ b/Engine/Plugins/Experimental/DataPrepEditor/Source/DataPrepLibraries/Private/DataprepOperations.h
@@ -161,7 +161,7 @@ protected:
 	//~ End UDataprepOperation Interface
 };
 
-UCLASS(Experimental, Category = ActorOperation, Meta = (DisplayName="Set Mobility", ToolTip = "For each mesh actor to process, update its mobilty with the selected value") )
+UCLASS(Experimental, Category = ActorOperation, Meta = (DisplayName="Set Mobility", ToolTip = "For each actor to process, update its mobilty with the selected value") )
 class UDataprepSetMobilityOperation : public UDataprepOperation
 {
 	GENERATED_BODY()
@@ -173,14 +173,14 @@ class UDataprepSetMobilityOperation : public UDataprepOperation
 
 public:
 	// Type of mobility to set on mesh actors
-	UPROPERTY(EditAnywhere, BlueprintReadWrite, Category = MeshOperation, meta = (ToolTip = "Type of mobility to set on mesh actors"))
+	UPROPERTY(EditAnywhere, BlueprintReadWrite, Category = ActorOperation, meta = (ToolTip = "Type of mobility to set on actors"))
 	TEnumAsByte<EComponentMobility::Type> MobilityType;
 
 	//~ Begin UDataprepOperation Interface
 public:
 	virtual FText GetCategory_Implementation() const override
 	{
-		return FDataprepOperationCategories::MeshOperation;
+		return FDataprepOperationCategories::ActorOperation;
 	}
 
 protected:
