Fixed a possible crash when editing a dataprep consumer.
The issue was that the ui was using a raw binding to a event from the comsumer and that binding was never removed when a widget died.

diff --git a/Engine/Plugins/Experimental/DataPrepEditor/Source/DataPrepEditor/Private/Widgets/DataprepWidgets.cpp b/Engine/Plugins/Experimental/DataPrepEditor/Source/DataPrepEditor/Private/Widgets/DataprepWidgets.cpp
--- a/Engine/Plugins/Experimental/DataPrepEditor/Source/DataPrepEditor/Private/Widgets/DataprepWidgets.cpp
+++ b/Engine/Plugins/Experimental/DataPrepEditor/Source/DataPrepEditor/Private/Widgets/DataprepWidgets.cpp
@@ -183,12 +183,12 @@ void SDataprepConsumerWidget::SetDataprepConsumer(UDataprepContentConsumer* InDa
 	{
 		if(DataprepConsumerPtr.IsValid())
 		{
-			DataprepConsumerPtr->GetOnChanged().RemoveAll( this );
+			DataprepConsumerPtr->GetOnChanged().Remove( OnConsumerChangedHandle );
 		}
 
 		DataprepConsumerPtr = InDataprepConsumer;
 
-		InDataprepConsumer->GetOnChanged().AddRaw( this, &SDataprepConsumerWidget::OnConsumerChanged );
+		OnConsumerChangedHandle = InDataprepConsumer->GetOnChanged().AddSP( this, &SDataprepConsumerWidget::OnConsumerChanged );
 
 		OnConsumerChanged();
 	}
@@ -222,6 +222,14 @@ void SDataprepConsumerWidget::Construct(const FArguments& InArgs )
 	}
 }
 
+SDataprepConsumerWidget::~SDataprepConsumerWidget()
+{
+	if ( UDataprepContentConsumer* DataprepConsumer = DataprepConsumerPtr.Get() )
+	{
+		DataprepConsumer->GetOnChanged().Remove( OnConsumerChangedHandle );
+	}
+}
+
 TSharedRef<SWidget> SDataprepConsumerWidget::BuildWidget()
 {
 	TSharedRef<SWidget> BrowseButton = PropertyCustomizationHelpers::MakeBrowseButton( FSimpleDelegate::CreateSP( this, &SDataprepConsumerWidget::OnBrowseContentFolder ) );
diff --git a/Engine/Plugins/Experimental/DataPrepEditor/Source/DataPrepEditor/Private/Widgets/DataprepWidgets.h b/Engine/Plugins/Experimental/DataPrepEditor/Source/DataPrepEditor/Private/Widgets/DataprepWidgets.h
--- a/Engine/Plugins/Experimental/DataPrepEditor/Source/DataPrepEditor/Private/Widgets/DataprepWidgets.h
+++ b/Engine/Plugins/Experimental/DataPrepEditor/Source/DataPrepEditor/Private/Widgets/DataprepWidgets.h
@@ -70,9 +70,10 @@ public:
 	SLATE_ARGUMENT(UDataprepContentConsumer*, DataprepConsumer)
 	SLATE_END_ARGS()
 
-public:
 	void Construct(const FArguments& InArgs );
 
+	virtual ~SDataprepConsumerWidget();
+
 	/** Update DataprepConsumerPtr and content folder's and level name's text boxes */
 	void SetDataprepConsumer( UDataprepContentConsumer* DataprepConsumer );
 
@@ -115,6 +116,8 @@ private:
 	TSharedPtr< FDataprepDetailsViewColumnSizeData > ColumnSizeData;
 	/** Relative width to control splitters. */
 	float ColumnWidth;
+
+	FDelegateHandle OnConsumerChangedHandle;
 };
 
 class SDataprepDetailsView : public SCompoundWidget, public FGCObject
