Dataprep Input/Output/Parametrization panel does not have UE4 look and feel

diff --git a/Engine/Plugins/Experimental/DataPrepEditor/Source/DataPrepEditor/Private/Widgets/DataprepAssetView.cpp b/Engine/Plugins/Experimental/DataPrepEditor/Source/DataPrepEditor/Private/Widgets/DataprepAssetView.cpp
--- a/Engine/Plugins/Experimental/DataPrepEditor/Source/DataPrepEditor/Private/Widgets/DataprepAssetView.cpp
+++ b/Engine/Plugins/Experimental/DataPrepEditor/Source/DataPrepEditor/Private/Widgets/DataprepAssetView.cpp
@@ -289,11 +289,116 @@ void SGraphNodeDetailsWidget::UpdateFromObjects(const TArray<UObject*>& Property
 	ContextualEditingBorderWidget->SetContent(ContextualEditingWidget);
 }
 
-void SDataprepAssetView::Construct( const FArguments& InArgs, UDataprepAssetInterface* InDataprepAssetPtr, TSharedPtr<FUICommandList>& CommandList )
+TSharedRef<ITableRow> SDataprepAssetView::OnGenerateRowForCategoryTree( TSharedRef<EDataprepCategory> InTreeNode, const TSharedRef<STableViewBase>& InOwnerTable )
+{
+	TSharedPtr<ITableRow> Row;
+
+	switch( InTreeNode.Get() )
+	{
+		case EDataprepCategory::Producers:
+		{
+			ProducersWidget = SNew( SDataprepProducersWidget, DataprepAssetInterfacePtr->GetProducers(), CommandList )
+				.ColumnSizeData( ColumnSizeData );
+
+			TSharedPtr< SHorizontalBox > AddNewProducerWrapper = SNew( SHorizontalBox )
+				+ SHorizontalBox::Slot()
+				.AutoWidth()
+				.Padding( 8.0f, 0.0f, 0.0f, 0.0f )
+				.HAlign( EHorizontalAlignment::HAlign_Right )
+				[
+					ProducersWidget->GetAddNewMenu().ToSharedRef()
+				];
+
+			Row = SNew( SDataprepCategoryWidget, ProducersWidget.ToSharedRef(), InOwnerTable )
+				.ColumnSizeData( ColumnSizeData )
+				.Title( LOCTEXT("DataprepProducersWidget_Producers_label", "Inputs") )
+				.TitleDetail( AddNewProducerWrapper.ToSharedRef() );
+			
+			break;
+		}
+		case EDataprepCategory::Consumers:
+		{
+			ConsumerWidget = SNew( SDataprepConsumerWidget )
+				.DataprepConsumer( DataprepAssetInterfacePtr->GetConsumer() )
+				.ColumnSizeData( ColumnSizeData );
+
+			TSharedPtr< SVerticalBox > ConsumerContainer = SNew( SVerticalBox )
+				+ SVerticalBox::Slot()
+				.AutoHeight()
+				[
+					ConsumerWidget.ToSharedRef()
+				]
+				+ SVerticalBox::Slot()
+				.AutoHeight()
+				[
+					SNew( SDataprepDetailsView )
+					.Object( DataprepAssetInterfacePtr->GetConsumer() )
+				];
+
+			TSharedPtr< SHorizontalBox > ConsumerSelectorWrapper = SNew( SHorizontalBox )
+				+ SHorizontalBox::Slot()
+				.AutoWidth()
+				.HAlign( EHorizontalAlignment::HAlign_Right )
+				[
+					ConsumerSelector.ToSharedRef()
+				];
+
+			Row = SNew( SDataprepCategoryWidget, ConsumerContainer.ToSharedRef(), InOwnerTable )
+				.ColumnSizeData( ColumnSizeData )
+				.Title( LOCTEXT("DataprepAssetView_Consumer_label", "Output") )
+				.TitleDetail( ConsumerSelectorWrapper.ToSharedRef() );
+
+			break;
+		}
+		case EDataprepCategory::Parameterization:
+		{
+			TSharedPtr<SDataprepDetailsView> ParameterizationDetailsView;
+
+			TSharedPtr< DataprepWidgetUtils::SConstrainedBox> ParametrizationContainer = SNew( DataprepWidgetUtils::SConstrainedBox )
+				[
+					SNew( SVerticalBox )
+					+ SVerticalBox::Slot()
+					.Padding( 8.0f, 5.0f, 0.0f, 0.0f )
+					.AutoHeight()
+					[
+
+						SAssignNew( ParameterizationDetailsView, SDataprepDetailsView )
+						.Object( DataprepAssetInterfacePtr->GetParameterizationObject() )
+						.ColumnSizeData( ColumnSizeData )
+						.Spacing( 10.0f )
+						.ColumnPadding( true )
+					]
+				];
+
+			if( DataprepAssetInterfacePtr->IsA<UDataprepAsset>() )
+			{
+				TWeakObjectPtr<UDataprepAsset> DataprepAsset = static_cast<UDataprepAsset*>( DataprepAssetInterfacePtr.Get() );
+				OnParameterizationWasEdited = DataprepAsset->OnParameterizedObjectsChanged.AddLambda( [ParameterizationDetailsView, DataprepAsset](const TSet<UObject*>* Objects )
+				{
+					if( Objects && Objects->Contains(DataprepAsset->GetParameterizationObject()) )
+					{
+						ParameterizationDetailsView->ForceRefresh();
+					}
+				});
+			}
+
+			Row = SNew( SDataprepCategoryWidget, ParametrizationContainer.ToSharedRef(), InOwnerTable )
+				.Title( LOCTEXT("DataprepAssetView_Consumer_Parameterization", "Parameterization") )
+				.ColumnSizeData( ColumnSizeData );
+			
+			break;
+		}
+	}
+
+	return Row.ToSharedRef();
+}
+
+void SDataprepAssetView::Construct( const FArguments& InArgs, UDataprepAssetInterface* InDataprepAssetPtr, TSharedPtr<FUICommandList>& InCommandList )
 {
 	check( InDataprepAssetPtr );
 
 	DataprepAssetInterfacePtr = InDataprepAssetPtr;
+	CommandList = InCommandList;
 
 	DataprepAssetInterfacePtr->GetOnChanged().AddRaw( this, &SDataprepAssetView::OnDataprepAssetChanged );
 
@@ -343,7 +448,6 @@ void SDataprepAssetView::Construct( const FArguments& InArgs, UDataprepAssetInte
 			SelectedConsumerDescription = MakeShared<FString>( FString() );
 		}
 
-
 		ConsumerSelector = SNew( STextComboBox )
 		.OptionsSource( &ConsumerDescriptionList )
 		.OnSelectionChanged( this, &SDataprepAssetView::OnNewConsumerSelected )
@@ -354,9 +458,6 @@ void SDataprepAssetView::Construct( const FArguments& InArgs, UDataprepAssetInte
 		ConsumerSelector = SNullWidget::NullWidget;
 	}
 
-	ProducersWidget = SNew(SDataprepProducersWidget, AssetProducers, CommandList)
-		.ColumnSizeData(ColumnSizeData);
-
 	TSharedRef<SDataprepDetailsView> DetailView = SNew(SDataprepDetailsView)
 	.ColumnSizeData( ColumnSizeData )
 	.Object( DataprepAssetInterfacePtr->GetParameterizationObject() );
@@ -422,8 +523,17 @@ void SDataprepAssetView::Construct( const FArguments& InArgs, UDataprepAssetInte
 			]
 		];
 	}
-
-	TSharedPtr<SDataprepDetailsView> ParameterizationDetailsView;
+	
+	Categories.Add( MakeShared<EDataprepCategory>( EDataprepCategory::Producers ) );
+	Categories.Add( MakeShared<EDataprepCategory>( EDataprepCategory::Consumers ) );
+	Categories.Add( MakeShared<EDataprepCategory>( EDataprepCategory::Parameterization ) );
+
+	TSharedRef< SDataprepCategoryTree > CategoryTree = SNew( SDataprepCategoryTree )
+		.TreeItemsSource( &Categories )
+		.OnGetChildren( this, &SDataprepAssetView::OnGetChildrenForCategoryTree )
+		.OnGenerateRow( this, &SDataprepAssetView::OnGenerateRowForCategoryTree )
+		.SelectionMode( ESelectionMode::None )
+		.HandleDirectionalNavigation( false );
 
 	ChildSlot
 	[
@@ -454,151 +564,11 @@ void SDataprepAssetView::Construct( const FArguments& InArgs, UDataprepAssetInte
 						ParentSpacer
 					]
 					// End - Section for Dataprep parent
-					// Begin - Section for producers
-					+ SVerticalBox::Slot()
-					.Padding( 5.0f )
-					.AutoHeight()
-					.MaxHeight( 400.f )
-					[
-						ProducersWidget.ToSharedRef()
-					]
-					// End - Section for producers
-
-					// Begin - Section for consumer
-					+ SVerticalBox::Slot()
-					.AutoHeight()
-					[
-						SNew(SHorizontalBox)
-						+ SHorizontalBox::Slot()
-						.FillWidth(1.0f)
-						.Padding( 0, 10, 0, 0 )
-						.HAlign( EHorizontalAlignment::HAlign_Center )
-						[
-							SNew(SColorBlock)
-							.Color( FLinearColor( 0.9f, 0.9f, 0.9f ) )
-							.Size( FVector2D( 3000, 1 ) )
-						]
-					]
-					+ SVerticalBox::Slot()
-					.AutoHeight()
-					[
-						SNew(SHorizontalBox)
-						+ SHorizontalBox::Slot()
-						[
-							SNew(SSpacer)
-							.Size( FVector2D( 200, 10 ) )		
-						]
-					]
-					+ SVerticalBox::Slot()
-					.AutoHeight()
-					[
-						SNew(SHorizontalBox)
-						+ SHorizontalBox::Slot()
-						.FillWidth(1.0f)
-						.HAlign(EHorizontalAlignment::HAlign_Left)
-						.VAlign(VAlign_Center)
-						.Padding( FMargin( 5.0f, 5.0f, 0, 7.0f ) )
-						[
-							SNew(STextBlock)
-							.Text(LOCTEXT("DataprepAssetView_Consumer_label", "Output"))
-							.MinDesiredWidth( 200 )
-							.Font( FDataprepEditorStyle::Get()->GetFontStyle("DataprepEditor.SectionFont") )
-						]
-						+ SHorizontalBox::Slot()
-						.AutoWidth()
-						.HAlign(EHorizontalAlignment::HAlign_Right)
-						.Padding(0, 0, 2, 0)
-						[
-							ConsumerSelector.ToSharedRef()
-						]
-					]
-					+ SVerticalBox::Slot()
-					.AutoHeight()
-					[
-						SAssignNew( ConsumerWidget, SDataprepConsumerWidget )
-						.DataprepConsumer( DataprepAssetInterfacePtr->GetConsumer() )
-						.ColumnSizeData( ColumnSizeData )
-					]
-					+ SVerticalBox::Slot()
-					.AutoHeight()
-					[
-						SNew( SDataprepDetailsView )
-						.Object( DataprepAssetInterfacePtr->GetConsumer())
-					]
 					+ SVerticalBox::Slot()
 					.AutoHeight()
 					[
-						SNew(SHorizontalBox)
-						+ SHorizontalBox::Slot()
-						[
-							SNew(SSpacer)
-							.Size( FVector2D( 200, 10 ) )		
-						]
+						CategoryTree
 					]
-					// End - Section for consumer
-
-					// Begin - Section for parameterization
-					+ SVerticalBox::Slot()
-					.AutoHeight()
-					[
-						SNew(SHorizontalBox)
-						+ SHorizontalBox::Slot()
-						.FillWidth(1.0f)
-						.Padding( 0, 10, 0, 0 )
-						.HAlign( EHorizontalAlignment::HAlign_Center )
-						[
-							SNew(SColorBlock)
-							.Color( FLinearColor( 0.9f, 0.9f, 0.9f ) )
-							.Size( FVector2D( 3000, 1 ) )
-						]
-					]
-					+ SVerticalBox::Slot()
-					.AutoHeight()
-					[
-						SNew(SHorizontalBox)
-						+ SHorizontalBox::Slot()
-						[
-							SNew(SSpacer)
-							.Size( FVector2D( 200, 10 ) )		
-						]
-					]
-					+ SVerticalBox::Slot()
-					.AutoHeight()
-					[
-						SNew(SHorizontalBox)
-						+ SHorizontalBox::Slot()
-						.FillWidth(1.0f)
-						.HAlign(EHorizontalAlignment::HAlign_Left)
-						.VAlign(VAlign_Center)
-						.Padding( FMargin( 5.0f, 5.0f, 0, 7.0f ) )
-						[
-							SNew(STextBlock)
-							.Text(LOCTEXT("DataprepAssetView_Consumer_Parameterization", "Parameterization"))
-							.MinDesiredWidth( 200 )
-							.Font( FDataprepEditorStyle::Get()->GetFontStyle("DataprepEditor.SectionFont") )
-						]
-					]
-					+ SVerticalBox::Slot()
-					.AutoHeight()
-					[
-						SNew(SHorizontalBox)
-						+ SHorizontalBox::Slot()
-						.FillWidth(1.0f)
-						.HAlign(EHorizontalAlignment::HAlign_Left)
-						.VAlign(VAlign_Center)
-						.Padding(10.0f, 5.0f, 0.0f, 5.0f)
-						[
-							SNew( DataprepWidgetUtils::SConstrainedBox )
-							[
-								SAssignNew( ParameterizationDetailsView, SDataprepDetailsView )
-								.Object( DataprepAssetInterfacePtr->GetParameterizationObject() )
-								.ColumnSizeData( ColumnSizeData )
-								.Spacing( 10.0f )
-								.ColumnPadding( true )
-							]
-						]
-					]
-					// End - Section for parameterization
 				]
 			]
 			+SHorizontalBox::Slot()
@@ -612,19 +582,6 @@ void SDataprepAssetView::Construct( const FArguments& InArgs, UDataprepAssetInte
 			]
 		]
 	];
-
-	if ( DataprepAssetInterfacePtr->IsA<UDataprepAsset>() )
-	{
-		TWeakObjectPtr<UDataprepAsset> DataprepAsset = static_cast<UDataprepAsset*>( DataprepAssetInterfacePtr.Get() );
-		OnParameterizationWasEdited = DataprepAsset->OnParameterizedObjectsChanged.AddLambda( [ParameterizationDetailsView, DataprepAsset](const TSet<UObject*>* Objects)
-		{
-			if ( Objects && Objects->Contains( DataprepAsset->GetParameterizationObject() ) )
-			{
-				ParameterizationDetailsView->ForceRefresh();
-			}
-		});
-	}
-
 }
 
 SDataprepAssetView::~SDataprepAssetView()
diff --git a/Engine/Plugins/Experimental/DataPrepEditor/Source/DataPrepEditor/Private/Widgets/DataprepAssetView.h b/Engine/Plugins/Experimental/DataPrepEditor/Source/DataPrepEditor/Private/Widgets/DataprepAssetView.h
--- a/Engine/Plugins/Experimental/DataPrepEditor/Source/DataPrepEditor/Private/Widgets/DataprepAssetView.h
+++ b/Engine/Plugins/Experimental/DataPrepEditor/Source/DataPrepEditor/Private/Widgets/DataprepAssetView.h
@@ -28,7 +28,7 @@ public:
 	SLATE_BEGIN_ARGS(SDataprepAssetView) {}
 	SLATE_END_ARGS()
 
-	void Construct(const FArguments& InArgs, UDataprepAssetInterface* InDataprepAssetPtr, TSharedPtr<FUICommandList>& CommandList);
+	void Construct(const FArguments& InArgs, UDataprepAssetInterface* InDataprepAssetPtr, TSharedPtr<FUICommandList>& InCommandList);
 
 	~SDataprepAssetView();
 
@@ -42,6 +42,12 @@ private:
 	float OnGetRightColumnWidth() const { return ColumnWidth; }
 	void OnSetColumnWidth(float InWidth) { ColumnWidth = InWidth; }
 
+	TSharedRef<ITableRow> OnGenerateRowForCategoryTree( TSharedRef<EDataprepCategory> InTreeNode, const TSharedRef<STableViewBase>& InOwnerTable );
+	void OnGetChildrenForCategoryTree( TSharedRef<EDataprepCategory> InTreeNode, TArray< TSharedRef<EDataprepCategory> >& OutChildren ) {}
+
+	TArray<TSharedRef< EDataprepCategory >> Categories;
+	TSharedPtr<FUICommandList> CommandList;
+
 private:
 	TWeakObjectPtr<UDataprepAssetInterface> DataprepAssetInterfacePtr;
 	TSharedPtr<SDataprepProducersWidget> ProducersWidget;
diff --git a/Engine/Plugins/Experimental/DataPrepEditor/Source/DataPrepEditor/Private/Widgets/DataprepWidgets.cpp b/Engine/Plugins/Experimental/DataPrepEditor/Source/DataPrepEditor/Private/Widgets/DataprepWidgets.cpp
--- a/Engine/Plugins/Experimental/DataPrepEditor/Source/DataPrepEditor/Private/Widgets/DataprepWidgets.cpp
+++ b/Engine/Plugins/Experimental/DataPrepEditor/Source/DataPrepEditor/Private/Widgets/DataprepWidgets.cpp
@@ -97,6 +97,100 @@ namespace DataprepWidgetUtils
 	}
 }
 
+void SDataprepCategoryWidget::ToggleExpansion()
+{
+	bIsExpanded = !bIsExpanded;
+	CategoryContent->SetVisibility(bIsExpanded ? EVisibility::Visible : EVisibility::Collapsed);
+}
+
+const FSlateBrush* SDataprepCategoryWidget::GetBackgroundImage() const
+{
+	if (IsHovered())
+	{
+		return bIsExpanded ? FEditorStyle::GetBrush("DetailsView.CategoryTop_Hovered") : FEditorStyle::GetBrush("DetailsView.CollapsedCategory_Hovered");
+	} 
+	else
+	{
+		return bIsExpanded ? FEditorStyle::GetBrush("DetailsView.CategoryTop") : FEditorStyle::GetBrush("DetailsView.CollapsedCategory");
+	}
+}
+
+void SDataprepCategoryWidget::Construct( const FArguments& InArgs, TSharedRef< SWidget > InContent, const TSharedRef<STableViewBase>& InOwnerTableView )
+{
+	const float MyContentTopPadding = 2.0f;
+	const float MyContentBottomPadding = 2.0f;
+
+	const float ChildSlotPadding = 2.0f;
+	const float BorderVerticalPadding = 3.0f;
+
+	CategoryContent = InContent;
+
+	TSharedPtr<SWidget> TitleDetail = InArgs._TitleDetail;
+	if( !TitleDetail.IsValid() )
+	{
+		TitleDetail = SNullWidget::NullWidget;
+	}
+
+	TSharedPtr<SHorizontalBox> TitleHeader = SNew(SHorizontalBox)
+		+ SHorizontalBox::Slot()
+		.VAlign( VAlign_Center )
+		.Padding( 2.0f, MyContentTopPadding, 2.0f, MyContentBottomPadding )
+		.AutoWidth()
+		[
+			SNew( SExpanderArrow, SharedThis(this) )
+		]
+		+ SHorizontalBox::Slot()
+		.VAlign( VAlign_Center )
+		.AutoWidth()
+		.Padding( 0.0f, 8.0f )
+		[
+			SNew(STextBlock)
+			.Text( InArgs._Title )
+			.Font( FEditorStyle::GetFontStyle("DetailsView.CategoryFontStyle") )
+			.ShadowOffset( FVector2D( 1.0f, 1.0f ) )
+		];
+
+	ChildSlot
+	[
+		SNew(SVerticalBox)
+		+ SVerticalBox::Slot()
+		.AutoHeight()
+		[
+			SNew( SBorder )
+			.BorderImage( this, &SDataprepCategoryWidget::GetBackgroundImage )
+			.Padding( FMargin( 0.0f, BorderVerticalPadding, 16.0f, BorderVerticalPadding ) )
+			.BorderBackgroundColor( FLinearColor( .6, .6, .6, 1.0f ) )
+			[
+				SNew(SHorizontalBox)
+				+ SHorizontalBox::Slot()
+				.FillWidth( 0.5f )
+				.HAlign( EHorizontalAlignment::HAlign_Left )
+				[
+					TitleHeader.ToSharedRef()
+				]
+				+ SHorizontalBox::Slot()
+				.FillWidth( 0.5f )
+				.HAlign( EHorizontalAlignment::HAlign_Right )
+				[
+					TitleDetail.ToSharedRef()
+				]
+			]
+		]
+		+ SVerticalBox::Slot()
+		.AutoHeight()
+		[
+			CategoryContent.ToSharedRef()
+		]
+	];
+
+	STableRow< TSharedPtr< EDataprepCategory > >::ConstructInternal(
+		STableRow::FArguments()
+		.Style( FEditorStyle::Get(), "DetailsView.TreeView.TableRow" )
+		.ShowSelection( false ),
+		InOwnerTableView
+	);
+}
+
 void SDataprepConsumerWidget::OnLevelNameChanged( const FText &NewLevelName, ETextCommit::Type CommitType )
 {
 	const FScopedTransaction Transaction( LOCTEXT("Consumer_SetLevelName", "Set Level Name") );
diff --git a/Engine/Plugins/Experimental/DataPrepEditor/Source/DataPrepEditor/Private/Widgets/DataprepWidgets.h b/Engine/Plugins/Experimental/DataPrepEditor/Source/DataPrepEditor/Private/Widgets/DataprepWidgets.h
--- a/Engine/Plugins/Experimental/DataPrepEditor/Source/DataPrepEditor/Private/Widgets/DataprepWidgets.h
+++ b/Engine/Plugins/Experimental/DataPrepEditor/Source/DataPrepEditor/Private/Widgets/DataprepWidgets.h
@@ -10,6 +10,7 @@
 #include "Framework/SlateDelegates.h"
 #include "Widgets/Layout/SSplitter.h"
 #include "Widgets/SCompoundWidget.h"
+#include "Widgets/Views/SListView.h"
 
 class SEditableTextBox;
 class SGridPanel;
@@ -59,6 +60,37 @@ struct FDataprepDetailsViewColumnSizeData
 	void SetColumnWidth(float InWidth) { OnWidthChanged.ExecuteIfBound(InWidth); }
 };
 
+enum class EDataprepCategory 
+{ 
+	Producers, 
+	Consumers, 
+	Parameterization 
+};
+
+typedef STreeView< TSharedRef<EDataprepCategory> > SDataprepCategoryTree;
+
+class SDataprepCategoryWidget : public STableRow< TSharedPtr< EDataprepCategory > >
+{
+public:
+
+	SLATE_BEGIN_ARGS(SDataprepCategoryWidget) {}
+	SLATE_ARGUMENT(TSharedPtr< FDataprepDetailsViewColumnSizeData >, ColumnSizeData)
+	SLATE_ARGUMENT(FText, Title)
+	SLATE_ARGUMENT(TSharedPtr< SWidget >, TitleDetail)
+	SLATE_END_ARGS()
+
+	void Construct( const FArguments& InArgs, TSharedRef< SWidget > InContent, const TSharedRef<STableViewBase>& InOwnerTableView );
+
+	virtual int32 DoesItemHaveChildren() const override { return 1; }
+	virtual bool IsItemExpanded() const override { return bIsExpanded; }
+	virtual void ToggleExpansion() override;
+
+	const FSlateBrush* GetBackgroundImage() const;
+private:
+	bool bIsExpanded = true;
+	TSharedPtr< SWidget > CategoryContent;
+};
+
 class SDataprepConsumerWidget : public SCompoundWidget
 {
 public:
@@ -70,7 +102,7 @@ public:
 	SLATE_ARGUMENT(UDataprepContentConsumer*, DataprepConsumer)
 	SLATE_END_ARGS()
 
-	void Construct(const FArguments& InArgs );
+	void Construct( const FArguments& InArgs );
 
 	virtual ~SDataprepConsumerWidget();
 
diff --git a/Engine/Plugins/Experimental/DataPrepEditor/Source/DataPrepEditor/Private/Widgets/SDataprepProducersWidget.cpp b/Engine/Plugins/Experimental/DataPrepEditor/Source/DataPrepEditor/Private/Widgets/SDataprepProducersWidget.cpp
--- a/Engine/Plugins/Experimental/DataPrepEditor/Source/DataPrepEditor/Private/Widgets/SDataprepProducersWidget.cpp
+++ b/Engine/Plugins/Experimental/DataPrepEditor/Source/DataPrepEditor/Private/Widgets/SDataprepProducersWidget.cpp
@@ -261,19 +261,19 @@ void SDataprepProducersWidget::Construct(const FArguments & InArgs, UDataprepAss
 
 	TreeView = SNew( SDataprepProducersTreeView, InAssetProducersPtr, InArgs._ColumnSizeData.ToSharedRef() );
 
-	TSharedPtr<SWidget> AddNewMenu = SNew(SComboButton)
-	.ComboButtonStyle(FEditorStyle::Get(), "ToolbarComboButton")
-	.ForegroundColor(FLinearColor::White)
-	.ToolTipText(LOCTEXT("AddNewToolTip", "Add a new producer."))
-	.OnGetMenuContent( this, &SDataprepProducersWidget::CreateAddProducerMenuWidget)
-	.HasDownArrow(false)
+	AddNewMenu = SNew(SComboButton)
+	.ComboButtonStyle( FEditorStyle::Get(), "ToolbarComboButton" )
+	.ForegroundColor( FLinearColor::White )
+	.ToolTipText( LOCTEXT("AddNewToolTip", "Add a new producer.") )
+	.OnGetMenuContent( this, &SDataprepProducersWidget::CreateAddProducerMenuWidget )
+	.HasDownArrow( false )
 	.ButtonContent()
 	[
 		SNew(SHorizontalBox)
 		+ SHorizontalBox::Slot()
 		.AutoWidth()
-		.Padding(FMargin(0, 1))
-		.HAlign(HAlign_Center)
+		.Padding( FMargin(0, 1) )
+		.HAlign( HAlign_Center )
 		[
 			SNew(STextBlock)
 			.Font( FDataprepEditorUtils::GetGlyphFont() )
@@ -294,28 +294,8 @@ void SDataprepProducersWidget::Construct(const FArguments & InArgs, UDataprepAss
 			.ExternalScrollbar(ScrollBar)
 			+ SScrollBox::Slot()
 			[
-				SNew(SVerticalBox)
 				// Section for producers
-				+ SVerticalBox::Slot()
-				.AutoHeight()
-				[
-					SNew(SHorizontalBox)
-					+ SHorizontalBox::Slot()
-					.FillWidth(1.0f)
-					.HAlign(EHorizontalAlignment::HAlign_Left)
-					.VAlign(VAlign_Center)
-					[
-						SNew(STextBlock)
-						.Font( FDataprepEditorStyle::Get()->GetFontStyle("DataprepEditor.SectionFont") )
-						.Text(LOCTEXT("DataprepProducersWidget_Producers_label", "Inputs"))
-					]
-					+ SHorizontalBox::Slot()
-					.AutoWidth()
-					.HAlign(EHorizontalAlignment::HAlign_Right)
-					[
-						AddNewMenu.ToSharedRef()
-					]
-				]
+				SNew(SVerticalBox)
 				+ SVerticalBox::Slot()
 				.AutoHeight()
 				[
diff --git a/Engine/Plugins/Experimental/DataPrepEditor/Source/DataPrepEditor/Private/Widgets/SDataprepProducersWidget.h b/Engine/Plugins/Experimental/DataPrepEditor/Source/DataPrepEditor/Private/Widgets/SDataprepProducersWidget.h
--- a/Engine/Plugins/Experimental/DataPrepEditor/Source/DataPrepEditor/Private/Widgets/SDataprepProducersWidget.h
+++ b/Engine/Plugins/Experimental/DataPrepEditor/Source/DataPrepEditor/Private/Widgets/SDataprepProducersWidget.h
@@ -117,6 +117,7 @@ public:
 
 	void Refresh() { TreeView->Refresh(); }
 
+	TSharedPtr<SWidget> GetAddNewMenu() { return AddNewMenu; }
 
 protected:
 	void OnAddProducer( UClass* ProducerClass );
@@ -125,6 +126,7 @@ private:
 	void OnDataprepProducersChanged(FDataprepAssetChangeType ChangeType, int32 Index);
 
 private:
+	TSharedPtr<SWidget> AddNewMenu;
 	TSharedPtr<FUICommandList> CommandList;
 	TWeakObjectPtr<UDataprepAssetProducers> AssetProducersPtr;
 	TSharedPtr<SDataprepProducersTreeView> TreeView;
