Fixed build breakage from CL #9957206.

#rnx

diff --git a/Engine/Plugins/Experimental/DataPrepEditor/Source/DataPrepEditor/Private/Widgets/DataprepWidgets.h b/Engine/Plugins/Experimental/DataPrepEditor/Source/DataPrepEditor/Private/Widgets/DataprepWidgets.h
--- a/Engine/Plugins/Experimental/DataPrepEditor/Source/DataPrepEditor/Private/Widgets/DataprepWidgets.h
+++ b/Engine/Plugins/Experimental/DataPrepEditor/Source/DataPrepEditor/Private/Widgets/DataprepWidgets.h
@@ -11,6 +11,7 @@
 #include "Widgets/Layout/SSplitter.h"
 #include "Widgets/SCompoundWidget.h"
 #include "Widgets/Views/SListView.h"
+#include "Widgets/Views/STreeView.h"
 
 class SEditableTextBox;
 class SGridPanel;
