Some objects were left in the old package when renaming a dataprep asset.

diff --git a/Engine/Plugins/Experimental/DataPrepEditor/Source/DataPrepCore/Private/DataPrepAsset.cpp b/Engine/Plugins/Experimental/DataPrepEditor/Source/DataPrepCore/Private/DataPrepAsset.cpp
--- a/Engine/Plugins/Experimental/DataPrepEditor/Source/DataPrepCore/Private/DataPrepAsset.cpp
+++ b/Engine/Plugins/Experimental/DataPrepEditor/Source/DataPrepCore/Private/DataPrepAsset.cpp
@@ -119,6 +119,26 @@ void UDataprepAsset::PostLoad()
 	}
 }
 
+bool UDataprepAsset::Rename(const TCHAR* NewName/* =nullptr */, UObject* NewOuter/* =nullptr */, ERenameFlags Flags/* =REN_None */)
+{
+	bool bWasRename = Super::Rename( NewName, NewOuter, Flags );
+	if ( bWasRename )
+	{
+		if ( Parameterization )
+		{
+			bWasRename &= Parameterization->OnAssetRename( Flags );
+		}
+
+		if ( DataprepRecipeBP && bWasRename )
+		{
+			// There shouldn't be a blueprint depending on us. Should be ok to just rename the generated class
+			bWasRename &= DataprepRecipeBP->RenameGeneratedClasses( NewName, NewOuter, Flags );
+		}
+	}
+
+	return bWasRename;
+}
+
 bool UDataprepAsset::CreateBlueprint()
 {
 	// Begin: Temp code for the nodes development
diff --git a/Engine/Plugins/Experimental/DataPrepEditor/Source/DataPrepCore/Private/Parameterization/DataprepParameterization.cpp b/Engine/Plugins/Experimental/DataPrepEditor/Source/DataPrepCore/Private/Parameterization/DataprepParameterization.cpp
--- a/Engine/Plugins/Experimental/DataPrepEditor/Source/DataPrepCore/Private/Parameterization/DataprepParameterization.cpp
+++ b/Engine/Plugins/Experimental/DataPrepEditor/Source/DataPrepCore/Private/Parameterization/DataprepParameterization.cpp
@@ -950,6 +950,8 @@ bool UDataprepParameterization::BindObjectProperty(UDataprepParameterizableObjec
 {
 	if ( Object && FDataprepParameterizationUtils::IsPropertyChainValid( PropertyChain ) && !Name.IsNone() )
 	{
+		Modify();
+
 		TSharedRef<FDataprepParameterizationBinding> Binding = MakeShared<FDataprepParameterizationBinding>( Object, PropertyChain );
 		void* AddressOfTheValueFromBinding;
 
@@ -984,7 +986,6 @@ bool UDataprepParameterization::BindObjectProperty(UDataprepParameterizableObjec
 					DataprepParameterization::PopulateValueTypeValidationData( PropertyFromParameterization, ValueTypeValidationData );
 					if ( ValueTypeValidationData == Binding->ValueTypeValidationData )
 					{
-						Modify();
 						BindingsContainer->Add( Binding, Name, BindingsToRemove );
 						bBindingWasAdded = true;
 					}
@@ -992,8 +993,6 @@ bool UDataprepParameterization::BindObjectProperty(UDataprepParameterizableObjec
 			}
 			else
 			{
-				Modify();
-
 				BindingsContainer->Add( Binding, Name, BindingsToRemove );
 			
 				UProperty* PropertyFromBinding = PropertyChain.Last().CachedProperty.Get();
@@ -1289,11 +1288,12 @@ void UDataprepParameterization::PrepareCustomClassForNewClassGeneration()
 {
 	if ( CustomContainerClass )
 	{
-		const FString OldClassName = MakeUniqueObjectName( CustomContainerClass->GetOuter(), CustomContainerClass->GetClass(), *FString::Printf(TEXT("%s_REINST"), *CustomContainerClass->GetName()) ).ToString();
+		const FString OldClassName = MakeUniqueObjectName( GetTransientPackage(), CustomContainerClass->GetClass(), *FString::Printf(TEXT("%s_REINST"), *CustomContainerClass->GetName()) ).ToString();
 		CustomContainerClass->ClassFlags |= CLASS_NewerVersionExists;
-		CustomContainerClass->SetFlags( RF_NewerVersionExists );
 		CustomContainerClass->ClearFlags( RF_Public | RF_Standalone );
-		CustomContainerClass->Rename( *OldClassName, GetTransientPackage(), REN_DontCreateRedirectors | REN_DoNotDirty);
+		CustomContainerClass->Rename( *OldClassName, GetTransientPackage(), REN_DontCreateRedirectors | REN_DoNotDirty );
+		CustomContainerClass->SetFlags( RF_Transient | RF_NewerVersionExists );
+		CustomContainerClass->GetDefaultObject()->SetFlags( RF_Transient | RF_NewerVersionExists );
 	}
 }
 
@@ -1436,6 +1436,7 @@ void UDataprepParameterization::PushParametrizationValueToBindings(FName Paramet
 
 				if ( bClassNeedUpdate )
 				{
+					Modify();
 					UpdateClass();
 				}
 
@@ -1459,6 +1460,17 @@ bool UDataprepParameterization::RemoveBinding(const TSharedRef<FDataprepParamete
 	return !ParameterOfRemovedBinding.IsNone();
 }
 
+bool UDataprepParameterization::OnAssetRename(ERenameFlags Flags)
+{
+	if ( CustomContainerClass )
+	{
+		const FString NewClassName = MakeUniqueObjectName( GetOutermost(), CustomContainerClass->GetClass(), *CustomContainerClass->GetName() ).ToString();
+		return CustomContainerClass->Rename( *NewClassName, GetOutermost(), Flags );
+	}
+
+	return true;
+}
+
 const FName UDataprepParameterization::MetadataClassGeneratorName( TEXT("DataprepCustomParameterizationGenerator") );
 
 
diff --git a/Engine/Plugins/Experimental/DataPrepEditor/Source/DataPrepCore/Private/Parameterization/DataprepParameterization.h b/Engine/Plugins/Experimental/DataPrepEditor/Source/DataPrepCore/Private/Parameterization/DataprepParameterization.h
--- a/Engine/Plugins/Experimental/DataPrepEditor/Source/DataPrepCore/Private/Parameterization/DataprepParameterization.h
+++ b/Engine/Plugins/Experimental/DataPrepEditor/Source/DataPrepCore/Private/Parameterization/DataprepParameterization.h
@@ -276,6 +276,13 @@ private:
 
 public:
 
+	/**
+	 * Update the package of the generated class
+	 * Return if it can be rename
+	 */
+	bool OnAssetRename(ERenameFlags Flags);
+
+
 	static const FName MetadataClassGeneratorName;
 
 private:
diff --git a/Engine/Plugins/Experimental/DataPrepEditor/Source/DataPrepCore/Public/DataPrepAsset.h b/Engine/Plugins/Experimental/DataPrepEditor/Source/DataPrepCore/Public/DataPrepAsset.h
--- a/Engine/Plugins/Experimental/DataPrepEditor/Source/DataPrepCore/Public/DataPrepAsset.h
+++ b/Engine/Plugins/Experimental/DataPrepEditor/Source/DataPrepCore/Public/DataPrepAsset.h
@@ -38,6 +38,7 @@ public:
 
 	// UObject interface
 	virtual void PostLoad() override;
+	virtual bool Rename(const TCHAR* NewName/* =nullptr */, UObject* NewOuter/* =nullptr */, ERenameFlags Flags/* =REN_None */) override;
 	// End of UObject interface
 
 	// UDataprepAssetInterface interface
