Changed logic to skip meshes for editor only or for visualization
Improved snapshot code

#rnx

diff --git a/Engine/Plugins/Experimental/DataPrepEditor/Source/DataPrepEditor/Private/DataprepSnapshot.cpp b/Engine/Plugins/Experimental/DataPrepEditor/Source/DataPrepEditor/Private/DataprepSnapshot.cpp
--- a/Engine/Plugins/Experimental/DataPrepEditor/Source/DataPrepEditor/Private/DataprepSnapshot.cpp
+++ b/Engine/Plugins/Experimental/DataPrepEditor/Source/DataPrepEditor/Private/DataprepSnapshot.cpp
@@ -210,15 +210,20 @@ namespace DataprepSnapshotUtil
 			MemAr << ObjectFlags;
 		}
 
-		// Serialize sub-objects' outer path
-		// Done in reverse order since an object can be the outer of the object
+		// Serialize sub-objects' outer path and name
+		// Done in reverse order since a sub-object can be the outer of another sub-object
 		// it depends on. Not the opposite
 		for(int32 Index = SubObjectsArray.Num() - 1; Index >= 0; --Index)
 		{
-			FSoftObjectPath SoftPath( SubObjectsArray[Index]->GetOuter() );
+			const UObject* SubObject = SubObjectsArray[Index];
+
+			FSoftObjectPath SoftPath( SubObject->GetOuter() );
 
 			FString SoftPathString = SoftPath.ToString();
 			MemAr << SoftPathString;
+
+			FString SubObjectName = SubObject->GetName();
+			MemAr << SubObjectName;
 		}
 
 		// Serialize sub-objects' content
@@ -290,6 +295,9 @@ namespace DataprepSnapshotUtil
 			FString SoftPathString;
 			MemAr << SoftPathString;
 
+			FString SubObjectName;
+			MemAr << SubObjectName;
+
 			const FSoftObjectPath SoftPath( SoftPathString );
 
 			UObject* NewOuter = SoftPath.ResolveObject();
@@ -298,7 +306,7 @@ namespace DataprepSnapshotUtil
 			UObject* SubObject = SubObjectsArray[Index];
 			if( NewOuter != SubObject->GetOuter() )
 			{
-				FDataprepCoreUtils::RenameObject( SubObject, nullptr, NewOuter );
+				FDataprepCoreUtils::RenameObject( SubObject, *SubObjectName, NewOuter );
 			}
 		}
 
diff --git a/Engine/Plugins/Experimental/DataPrepEditor/Source/DataPrepLibraries/Private/DataprepEditingOperations.cpp b/Engine/Plugins/Experimental/DataPrepEditor/Source/DataPrepLibraries/Private/DataprepEditingOperations.cpp
--- a/Engine/Plugins/Experimental/DataPrepEditor/Source/DataPrepLibraries/Private/DataprepEditingOperations.cpp
+++ b/Engine/Plugins/Experimental/DataPrepEditor/Source/DataPrepLibraries/Private/DataprepEditingOperations.cpp
@@ -468,7 +468,7 @@ void UDataprepCreateProxyMeshOperation::OnExecution_Implementation(const FDatapr
 	// Then delete the merged actors that don't have any children component
 	for(AActor* Actor : ActorsToMerge)
 	{
-		UPrimitiveComponent* RootComponent = Cast<UPrimitiveComponent>(Actor->GetRootComponent());
+		USceneComponent* RootComponent = Cast<USceneComponent>(Actor->GetRootComponent());
 		if(RootComponent && RootComponent->GetNumChildrenComponents() == 0)
 		{
 			ObjectsToDelete.Add(Actor);
@@ -774,10 +774,14 @@ namespace DatasmithEditingOperationsUtils
 					bool bMeshActorIsValid = false;
 					for(UStaticMeshComponent* MeshComponent : ComponentArray)
 					{
-						if(MeshComponent->GetStaticMesh() && MeshComponent->GetStaticMesh()->GetSourceModels().Num() > 0)
+						// Skip components which are either editor only or for visualization
+						if(!MeshComponent->IsEditorOnly() && !MeshComponent->IsVisualizationComponent())
 						{
-							bMeshActorIsValid = true;
-							ComponentsToMerge.Add(MeshComponent);
+							if(MeshComponent->GetStaticMesh() && MeshComponent->GetStaticMesh()->GetSourceModels().Num() > 0)
+							{
+								bMeshActorIsValid = true;
+								ComponentsToMerge.Add(MeshComponent);
+							}
 						}
 					}
 
