Fixed a ui refresh issue when doing a undo/redo in a dataprep instance. (this could have cause some crashes)

diff --git a/Engine/Plugins/Experimental/DataPrepEditor/Source/DataPrepEditor/Private/Widgets/DataprepWidgets.cpp b/Engine/Plugins/Experimental/DataPrepEditor/Source/DataPrepEditor/Private/Widgets/DataprepWidgets.cpp
--- a/Engine/Plugins/Experimental/DataPrepEditor/Source/DataPrepEditor/Private/Widgets/DataprepWidgets.cpp
+++ b/Engine/Plugins/Experimental/DataPrepEditor/Source/DataPrepEditor/Private/Widgets/DataprepWidgets.cpp
@@ -49,6 +49,7 @@
 #include "Widgets/SNullWidget.h"
 #include "Widgets/Text/SInlineEditableTextBlock.h"
 #include "Widgets/Text/STextBlock.h"
+#include "UObject/UObjectGlobals.h"
 
 #define LOCTEXT_NAMESPACE "DataprepSlateHelper"
 
@@ -690,6 +691,15 @@ void SDataprepDetailsView::OnDataprepParameterizationStatusForObjectsChanged(con
 	}
 }
 
+void SDataprepDetailsView::OnObjectTransacted(UObject* Object, const class FTransactionObjectEvent& TransactionObjectEvent)
+{
+	// Hack to support refresh the parameterization display of a dataprep instance
+	if ( Object == DetailedObject || ( DetailedObject && DetailedObject->GetOuter() == Object ) )
+	{
+		ForceRefresh();
+	}
+}
+
 void SDataprepDetailsView::AddWidgets( const TArray< TSharedRef< IDetailTreeNode > >& DetailTree, int32& Index, float LeftPadding, const FDataprepParameterizationContext& InParameterizationContext)
 {
 	auto IsDetailNodeDisplayable = []( const TSharedPtr< IPropertyHandle >& PropertyHandle)
@@ -888,6 +898,7 @@ void SDataprepDetailsView::Construct(const FArguments& InArgs)
 	if ( GEditor )
 	{
 		OnObjectReplacedHandle = GEditor->OnObjectsReplaced().AddSP(this, &SDataprepDetailsView::OnObjectReplaced);
+		OnObjectTransactedHandle = FCoreUObjectDelegates::OnObjectTransacted.AddSP( this, &SDataprepDetailsView::OnObjectTransacted );
 	}
 
 	Construct();
@@ -965,6 +976,7 @@ SDataprepDetailsView::~SDataprepDetailsView()
 	if ( GEditor )
 	{
 		GEditor->OnObjectsReplaced().Remove( OnObjectReplacedHandle );
+		FCoreUObjectDelegates::OnObjectTransacted.Remove( OnObjectTransactedHandle );
 	}
 
 	if ( UDataprepAsset* DataprepAsset = DataprepAssetForParameterization.Get() )
diff --git a/Engine/Plugins/Experimental/DataPrepEditor/Source/DataPrepEditor/Private/Widgets/DataprepWidgets.h b/Engine/Plugins/Experimental/DataPrepEditor/Source/DataPrepEditor/Private/Widgets/DataprepWidgets.h
--- a/Engine/Plugins/Experimental/DataPrepEditor/Source/DataPrepEditor/Private/Widgets/DataprepWidgets.h
+++ b/Engine/Plugins/Experimental/DataPrepEditor/Source/DataPrepEditor/Private/Widgets/DataprepWidgets.h
@@ -233,6 +233,8 @@ private:
 
 	void OnDataprepParameterizationStatusForObjectsChanged(const TSet<UObject*>* Objects);
 
+	void OnObjectTransacted(UObject* Object, const class FTransactionObjectEvent& TransactionObjectEvent);
+
 private:
 	/** Row generator applied on detailed object */
 	TSharedPtr< class IPropertyRowGenerator > Generator;
@@ -249,9 +251,12 @@ private:
 	/** Delegate handle to track property changes on array properties */
 	FDelegateHandle OnPropertyChangedHandle;
 
-	/** Delegate to track new object after a re-instancing process */
+	/** Delegate handle to track new object after a re-instancing process */
 	FDelegateHandle OnObjectReplacedHandle;
 
+	/** Delegate handle to track when a object was transacted */
+	FDelegateHandle OnObjectTransactedHandle;
+
 	/** Relative width to control splitters. */
 	float ColumnWidth;
 
