Fixed offset in Producers' widget

diff --git a/Engine/Plugins/Experimental/DataPrepEditor/Source/DataPrepEditor/Private/Widgets/DataprepAssetView.cpp b/Engine/Plugins/Experimental/DataPrepEditor/Source/DataPrepEditor/Private/Widgets/DataprepAssetView.cpp
--- a/Engine/Plugins/Experimental/DataPrepEditor/Source/DataPrepEditor/Private/Widgets/DataprepAssetView.cpp
+++ b/Engine/Plugins/Experimental/DataPrepEditor/Source/DataPrepEditor/Private/Widgets/DataprepAssetView.cpp
@@ -300,6 +300,14 @@ TSharedRef<ITableRow> SDataprepAssetView::OnGenerateRowForCategoryTree( TSharedR
 			ProducersWidget = SNew( SDataprepProducersWidget, DataprepAssetInterfacePtr->GetProducers(), CommandList )
 				.ColumnSizeData( ColumnSizeData );
 
+			TSharedPtr< SWidget > ProducerWrapper = SNew( SHorizontalBox )
+				+ SHorizontalBox::Slot()
+				.FillWidth(1.0f)
+				.Padding( 5.0f, 0.0f, 0.0f, 0.0f )
+				[
+					ProducersWidget.ToSharedRef()
+				];
+
 			TSharedPtr< SHorizontalBox > AddNewProducerWrapper = SNew( SHorizontalBox )
 				+ SHorizontalBox::Slot()
 				.AutoWidth()
@@ -309,7 +317,7 @@ TSharedRef<ITableRow> SDataprepAssetView::OnGenerateRowForCategoryTree( TSharedR
 					ProducersWidget->GetAddNewMenu().ToSharedRef()
 				];
 
-			Row = SNew( SDataprepCategoryWidget, ProducersWidget.ToSharedRef(), InOwnerTable )
+			Row = SNew( SDataprepCategoryWidget, ProducerWrapper.ToSharedRef(), InOwnerTable )
 				.ColumnSizeData( ColumnSizeData )
 				.Title( LOCTEXT("DataprepProducersWidget_Producers_label", "Inputs") )
 				.TitleDetail( AddNewProducerWrapper.ToSharedRef() );
