Dataprep added a better validation before executing the fetch function of a fetcher in a parrallel for.
Also the parallel for now call directly the native function of the fetcher.

diff --git a/Engine/Plugins/Experimental/DataPrepEditor/Source/DataPrepCore/Private/SelectionSystem/DataprepSelectionSystemUtils.h b/Engine/Plugins/Experimental/DataPrepEditor/Source/DataPrepCore/Private/SelectionSystem/DataprepSelectionSystemUtils.h
--- a/Engine/Plugins/Experimental/DataPrepEditor/Source/DataPrepCore/Private/SelectionSystem/DataprepSelectionSystemUtils.h
+++ b/Engine/Plugins/Experimental/DataPrepEditor/Source/DataPrepCore/Private/SelectionSystem/DataprepSelectionSystemUtils.h
@@ -8,6 +8,8 @@
 #include "Async/ParallelFor.h"
 #include "CoreMinimal.h"
 #include "Templates/UnrealTypeTraits.h"
+#include "UObject/Class.h"
+#include "UObject/ObjectMacros.h"
 
 namespace DataprepSelectionSystemUtils
 {
@@ -42,12 +44,12 @@ namespace DataprepSelectionSystemUtils
 			}
 		};
 
-		if ( Filter.IsThreadSafe() && Fetcher.IsThreadSafe() )
+		if ( Filter.IsThreadSafe() && Fetcher.IsThreadSafe() && bool( Fetcher.GetClass()->ClassFlags & CLASS_Native ) )
 		{
 			ParallelFor( Objects.Num(), [&Fetcher, &Objects, &UpdateFilterResult](int32 Index)
 			{
 				bool bFetchSucceded = false;
-				const FetchedDataType FetchedData = Fetcher.Fetch( Objects[Index], bFetchSucceded );
+				const FetchedDataType FetchedData = Fetcher.Fetch_Implementation( Objects[Index], bFetchSucceded );
 				UpdateFilterResult( Index, bFetchSucceded, FetchedData );
 			});
 		}
@@ -69,8 +71,8 @@ namespace DataprepSelectionSystemUtils
 				UpdateFilterResult( Index, FetcherResult.Key, FetcherResult.Value );
 			});
 		}
-
-		else if ( Fetcher.IsThreadSafe() )
+		
+		else if ( Fetcher.IsThreadSafe() && bool( Fetcher.GetClass()->ClassFlags & CLASS_Native ) )
 		{
 			TArray< TPair< bool, FetchedDataType > > FetcherResults;
 			FetcherResults.AddZeroed( Objects.Num() );
@@ -78,7 +80,7 @@ namespace DataprepSelectionSystemUtils
 			ParallelFor( Objects.Num(), [&Fetcher, &Objects, &FetcherResults](int32 Index)
 			{
 				bool bFetchSucceded = false;
-				const FetchedDataType FetchedData = Fetcher.Fetch( Objects[Index], bFetchSucceded );
+				const FetchedDataType FetchedData = Fetcher.Fetch_Implementation( Objects[Index], bFetchSucceded );
 				FetcherResults[Index] = TPair< bool, FetchedDataType >( bFetchSucceded, FetchedData );
 			});
 
