Updated the default layout of the dataprep editor.

diff --git a/Engine/Plugins/Experimental/DataPrepEditor/Source/DataPrepEditor/Private/DataPrepEditor.cpp b/Engine/Plugins/Experimental/DataPrepEditor/Source/DataPrepEditor/Private/DataPrepEditor.cpp
--- a/Engine/Plugins/Experimental/DataPrepEditor/Source/DataPrepEditor/Private/DataPrepEditor.cpp
+++ b/Engine/Plugins/Experimental/DataPrepEditor/Source/DataPrepEditor/Private/DataPrepEditor.cpp
@@ -951,7 +951,7 @@ TSharedRef<SDockTab> FDataprepEditor::SpawnTabSceneViewport( const FSpawnTabArgs
 
 TSharedRef<FTabManager::FLayout> FDataprepEditor::CreateDataprepLayout()
 {
-	return FTabManager::NewLayout("Standalone_DataprepEditor_Layout_v0.6")
+	return FTabManager::NewLayout("Standalone_DataprepEditor_Layout_v0.7")
 		->AddArea
 		(
 			FTabManager::NewPrimaryArea()->SetOrientation(Orient_Vertical)
@@ -961,8 +961,6 @@ TSharedRef<FTabManager::FLayout> FDataprepEditor::CreateDataprepLayout()
 				->SetSizeCoefficient(0.1f)
 				->SetHideTabWell(true)
 				->AddTab(GetToolbarTabId(), ETabState::OpenedTab)
-				// Don't want the secondary toolbar tab to be opened if there's nothing in it
-				//->AddTab(SecondaryToolbarTabId, ETabState::ClosedTab)
 			)
 			->Split
 			(
@@ -970,31 +968,30 @@ TSharedRef<FTabManager::FLayout> FDataprepEditor::CreateDataprepLayout()
 				->Split
 				(
 					FTabManager::NewSplitter()->SetOrientation(Orient_Vertical)
-					->SetSizeCoefficient(0.9f)
 					->Split
 					(
 						FTabManager::NewSplitter()->SetOrientation(Orient_Horizontal)
-						->SetSizeCoefficient(0.5f)
+						->SetSizeCoefficient(0.75f)
 						->Split
 						(
 							FTabManager::NewStack()
-							->SetSizeCoefficient(0.3f)
-							->AddTab(ScenePreviewTabId, ETabState::OpenedTab)
+							->SetSizeCoefficient(0.2f)
+							->AddTab(AssetPreviewTabId, ETabState::OpenedTab)
 							->SetHideTabWell( true )
 						)
 						->Split
 						(
 							FTabManager::NewStack()
-							->SetSizeCoefficient(0.3f)
-							->AddTab(AssetPreviewTabId, ETabState::OpenedTab)
-							->SetHideTabWell( true )
+							->SetSizeCoefficient(0.55f)
+							->AddTab(SceneViewportTabId, ETabState::OpenedTab)
+							->SetHideTabWell(true)
 						)
 						->Split
 						(
 							FTabManager::NewStack()
-							->SetSizeCoefficient(0.4f)
-							->AddTab(SceneViewportTabId, ETabState::OpenedTab)
-							->SetHideTabWell(true)
+							->SetSizeCoefficient(0.25f)
+							->AddTab(ScenePreviewTabId, ETabState::OpenedTab)
+							->SetHideTabWell( true )
 						)
 					)
 					->Split
@@ -1025,14 +1022,14 @@ TSharedRef<FTabManager::FLayout> FDataprepEditor::CreateDataprepLayout()
 					->Split
 					(
 						FTabManager::NewStack()
-						->SetSizeCoefficient(0.25f)
+						->SetSizeCoefficient(0.3f)
 						->AddTab(DataprepAssetTabId, ETabState::OpenedTab)
 						->SetHideTabWell(true)
 					)
 					->Split
 					(
 						FTabManager::NewStack()
-						->SetSizeCoefficient(0.75f)
+						->SetSizeCoefficient(0.7f)
 						->AddTab(DetailsTabId, ETabState::OpenedTab)
 						->SetHideTabWell(true)
 					)
