Fixed a crash when using a bounding volume filter in dataprep on a static mesh component without a static mesh.

#rnx

diff --git a/Engine/Plugins/Experimental/DataPrepEditor/Source/DataPrepLibraries/Private/Fetchers/DataprepFloatFetcherLibrary.cpp b/Engine/Plugins/Experimental/DataPrepEditor/Source/DataPrepLibraries/Private/Fetchers/DataprepFloatFetcherLibrary.cpp
--- a/Engine/Plugins/Experimental/DataPrepEditor/Source/DataPrepLibraries/Private/Fetchers/DataprepFloatFetcherLibrary.cpp
+++ b/Engine/Plugins/Experimental/DataPrepEditor/Source/DataPrepLibraries/Private/Fetchers/DataprepFloatFetcherLibrary.cpp
@@ -5,13 +5,12 @@
 #include "Components/StaticMeshComponent.h"
 #include "Engine/StaticMesh.h"
 #include "GameFramework/Actor.h"
-#include "StaticMeshAttributes.h"
 #include "Misc/Optional.h"
 #include "UObject/Object.h"
 
 #define LOCTEXT_NAMESPACE "DataprepFloatFetcherLibrary"
 
-/* UDataprepFloatVolumeFetcher methods
+/* UDataprepFloatVolumeFetcher methods+
  *****************************************************************************/
 float UDataprepFloatBoundingVolumeFetcher::Fetch_Implementation(const UObject* Object, bool& bOutFetchSucceded) const
 {
@@ -19,22 +18,9 @@ float UDataprepFloatBoundingVolumeFetcher::Fetch_Implementation(const UObject* O
 	{
 		auto GetStaticMeshBoundingBox = [](const UStaticMesh* StaticMesh) -> FBox
 			{
-				if ( !StaticMesh->RenderData || !StaticMesh->RenderData->IsInitialized() )
+				if ( !StaticMesh )
 				{
-					FBox Box(ForceInit);
-					if ( const FMeshDescription* MeshDescription = StaticMesh->GetMeshDescription( 0 ) )
-					{
-						FStaticMeshConstAttributes Attributes(*MeshDescription);
-						TVertexAttributesConstRef<FVector> VertexPositions = Attributes.GetVertexPositions();
-						for ( const FVertexID VertexID : MeshDescription->Vertices().GetElementIDs() )
-						{
-							if ( !MeshDescription->IsVertexOrphaned( VertexID ) )
-							{
-								Box += VertexPositions[VertexID];
-							}
-						}
-					}
-					return Box;
+					return {};
 				}
 
 				return StaticMesh->GetBoundingBox();
@@ -55,7 +41,7 @@ float UDataprepFloatBoundingVolumeFetcher::Fetch_Implementation(const UObject* O
 						ComponentBox = GetStaticMeshBoundingBox( StaticMeshComponent->GetStaticMesh() );
 						ComponentBox.TransformBy( PrimComp->GetComponentToWorld() );
 					}
-					if ( PrimComp->IsRegistered() &&  PrimComp->IsCollisionEnabled() )
+					if ( PrimComp->IsRegistered() )
 					{
 						ActorBox += ComponentBox.IsValid? ComponentBox : PrimComp->Bounds.GetBox();
 					}
