Fixed the ux issue with the dataprep viewport always moving its camera on each execute.

#rnx

diff --git a/Engine/Plugins/Experimental/DataPrepEditor/Source/DataPrepEditor/Private/DataPrepEditor.cpp b/Engine/Plugins/Experimental/DataPrepEditor/Source/DataPrepEditor/Private/DataPrepEditor.cpp
--- a/Engine/Plugins/Experimental/DataPrepEditor/Source/DataPrepEditor/Private/DataPrepEditor.cpp
+++ b/Engine/Plugins/Experimental/DataPrepEditor/Source/DataPrepEditor/Private/DataPrepEditor.cpp
@@ -564,6 +564,8 @@ void FDataprepEditor::OnBuildWorld()
 	TakeSnapshot();
 
 	UpdatePreviewPanels();
+	SceneViewportView->FocusViewportOnScene();
+
 	bWorldBuilt = true;
 	bIsFirstRun = true;
 }
diff --git a/Engine/Plugins/Experimental/DataPrepEditor/Source/DataPrepEditor/Private/Widgets/SDataprepEditorViewport.cpp b/Engine/Plugins/Experimental/DataPrepEditor/Source/DataPrepEditor/Private/Widgets/SDataprepEditorViewport.cpp
--- a/Engine/Plugins/Experimental/DataPrepEditor/Source/DataPrepEditor/Private/Widgets/SDataprepEditorViewport.cpp
+++ b/Engine/Plugins/Experimental/DataPrepEditor/Source/DataPrepEditor/Private/Widgets/SDataprepEditorViewport.cpp
@@ -480,13 +480,16 @@ void SDataprepEditorViewport::UpdateScene()
 
 	PreviewScene->SetFloorOffset(-SceneBounds.Min.Z );
 
-	EditorViewportClient->FocusViewportOnBox( SceneBounds );
-
 	UpdateOverlayText();
 
 	SceneViewport->Invalidate();
 }
 
+void SDataprepEditorViewport::FocusViewportOnScene()
+{
+	EditorViewportClient->FocusViewportOnBox( SceneBounds );
+}
+
 TSharedRef<SEditorViewport> SDataprepEditorViewport::GetViewportWidget()
 {
 	return SharedThis(this);
diff --git a/Engine/Plugins/Experimental/DataPrepEditor/Source/DataPrepEditor/Private/Widgets/SDataprepEditorViewport.h b/Engine/Plugins/Experimental/DataPrepEditor/Source/DataPrepEditor/Private/Widgets/SDataprepEditorViewport.h
--- a/Engine/Plugins/Experimental/DataPrepEditor/Source/DataPrepEditor/Private/Widgets/SDataprepEditorViewport.h
+++ b/Engine/Plugins/Experimental/DataPrepEditor/Source/DataPrepEditor/Private/Widgets/SDataprepEditorViewport.h
@@ -165,6 +165,9 @@ public:
 	/** UpdateScene the viewport with the current content of the editor's preview world */
 	void UpdateScene();
 
+	/** Move the camera so that we can see the full scene */
+	void FocusViewportOnScene();
+
 	/** Delete all the rendering data created by the viewport for display */
 	void ClearScene();
 
