Fix merge issues in DataPrepEditor.cpp
#rnx

diff --git a/Engine/Plugins/Experimental/DataPrepEditor/Source/DataPrepEditor/Private/DataPrepEditor.cpp b/Engine/Plugins/Experimental/DataPrepEditor/Source/DataPrepEditor/Private/DataPrepEditor.cpp
--- a/Engine/Plugins/Experimental/DataPrepEditor/Source/DataPrepEditor/Private/DataPrepEditor.cpp
+++ b/Engine/Plugins/Experimental/DataPrepEditor/Source/DataPrepEditor/Private/DataPrepEditor.cpp
@@ -44,6 +44,7 @@
 #include "Materials/MaterialInstance.h"
 #include "Misc/Attribute.h"
 #include "Misc/Guid.h"
+#include "Misc/MessageDialog.h"
 #include "Misc/Paths.h"
 #include "Misc/ScopedSlowTask.h"
 #include "Modules/ModuleManager.h"
@@ -842,7 +843,7 @@ void FDataprepEditor::OnCommitWorld()
 		const FText Title( LOCTEXT( "DataprepEditor_ProceedWithCommit", "Proceed with commit" ) );
 		const FText Message( LOCTEXT( "DataprepEditor_ConfirmCommitPipelineNotExecuted", "The action pipeline has not been executed.\nDo you want to proceeed with the commit anyway?" ) );
 
-		if( OpenMsgDlgInt( EAppMsgType::YesNo, Message, Title ) == EAppReturnType::No )
+		if( FMessageDialog::Open( EAppMsgType::YesNo, Message, &Title ) == EAppReturnType::No )
 		{
 			return;
 		}
