Temp fix for DataprepParametrization serialization crash.

#rnx

diff --git a/Engine/Plugins/Enterprise/DataprepEditor/Source/DataprepCore/Private/Parameterization/DataprepParameterization.cpp b/Engine/Plugins/Enterprise/DataprepEditor/Source/DataprepCore/Private/Parameterization/DataprepParameterization.cpp
--- a/Engine/Plugins/Enterprise/DataprepEditor/Source/DataprepCore/Private/Parameterization/DataprepParameterization.cpp
+++ b/Engine/Plugins/Enterprise/DataprepEditor/Source/DataprepCore/Private/Parameterization/DataprepParameterization.cpp
@@ -516,7 +516,7 @@ FDataprepParameterizationBinding::FDataprepParameterizationBinding(UDataprepPara
 
 bool FDataprepParameterizationBinding::Serialize(FArchive& Ar)
 {
-	check(false); // @todo FProperties: if we never hit this, we don't need this
+	//check(false); // @todo FProperties: if we never hit this, we don't need this
 
 	Ar.UsingCustomVersion(FCoreObjectVersion::GUID);
 
