Copy up Dev-Enterprise @ CL11133089

#rnx

diff --git a/Engine/Plugins/Enterprise/DataprepEditor/Source/DataprepCore/Private/DataprepActionAsset.cpp b/Engine/Plugins/Enterprise/DataprepEditor/Source/DataprepCore/Private/DataprepActionAsset.cpp
--- a/Engine/Plugins/Enterprise/DataprepEditor/Source/DataprepCore/Private/DataprepActionAsset.cpp
+++ b/Engine/Plugins/Enterprise/DataprepEditor/Source/DataprepCore/Private/DataprepActionAsset.cpp
@@ -53,6 +53,7 @@ void UDataprepActionStep::PostLoad()
 
 UDataprepActionAsset::UDataprepActionAsset()
 	: ContextPtr( nullptr )
+	, Label(TEXT("New Action"))
 {
 	bExecutionInterrupted = false;
 
@@ -194,6 +195,33 @@ int32 UDataprepActionAsset::AddStep(const UDataprepActionStep* InActionStep)
 	return INDEX_NONE;
 }
 
+int32 UDataprepActionAsset::AddSteps(const TArray<const UDataprepActionStep*>& InActionSteps)
+{
+	if ( InActionSteps.Num() > 0 && InActionSteps[0] != nullptr )
+	{
+		Modify();
+
+		Steps.Reserve(Steps.Num() + InActionSteps.Num());
+
+		for(const UDataprepActionStep* InActionStep : InActionSteps)
+		{
+			if(InActionStep)
+			{
+				UDataprepActionStep* ActionStep = DuplicateObject<UDataprepActionStep>( InActionStep, this );
+				Steps.Add( ActionStep );
+			}
+		}
+
+		OnStepsChanged.Broadcast();
+		return Steps.Num() - 1;
+	}
+
+	UE_LOG( LogDataprepCore, Error, TEXT("UDataprepActionAsset::AddSteps: The array is empty or the action steps are invalid") );
+	ensure(false);
+	// Invalid
+	return INDEX_NONE;
+}
+
 int32 UDataprepActionAsset::AddStep(const UDataprepParameterizableObject* StepObject)
 {
 	if (!StepObject)
@@ -244,6 +272,40 @@ bool UDataprepActionAsset::InsertStep(const UDataprepActionStep* InActionStep, i
 	return false;
 }
 
+bool UDataprepActionAsset::InsertSteps(const TArray<const UDataprepActionStep*>& InActionSteps, int32 Index)
+{
+	if ( !Steps.IsValidIndex( Index ) )
+	{
+		UE_LOG( LogDataprepCore, Error, TEXT("UDataprepActionAsset::InsertSteps: The Index is out of range") );
+		ensure( false );
+		return false;
+	}
+
+	if ( InActionSteps.Num() > 0 && InActionSteps[0] != nullptr )
+	{
+		Modify();
+
+		Steps.Reserve(Steps.Num() + InActionSteps.Num());
+
+		for(const UDataprepActionStep* InActionStep : InActionSteps)
+		{
+			if(InActionStep)
+			{
+				UDataprepActionStep* ActionStep = DuplicateObject<UDataprepActionStep>( InActionStep, this );
+				Steps.Insert( ActionStep, Index );
+			}
+		}
+
+		OnStepsChanged.Broadcast();
+		return true;
+	}
+
+	UE_LOG( LogDataprepCore, Error, TEXT("UDataprepActionAsset::InsertSteps: The array is empty or the action steps are invalid") );
+	ensure(false);
+	// Invalid
+	return false;
+}
+
 TWeakObjectPtr<UDataprepActionStep> UDataprepActionAsset::GetStep(int32 Index)
 {
 	// Avoid code duplication
diff --git a/Engine/Plugins/Enterprise/DataprepEditor/Source/DataprepCore/Private/DataprepAsset.cpp b/Engine/Plugins/Enterprise/DataprepEditor/Source/DataprepCore/Private/DataprepAsset.cpp
--- a/Engine/Plugins/Enterprise/DataprepEditor/Source/DataprepCore/Private/DataprepAsset.cpp
+++ b/Engine/Plugins/Enterprise/DataprepEditor/Source/DataprepCore/Private/DataprepAsset.cpp
@@ -34,10 +34,11 @@
 
 UDataprepAsset::UDataprepAsset()
 {
-	// Temp code for the nodes development
+#ifndef DP_NOBLUEPRINT
 	DataprepRecipeBP = nullptr;
 	StartNode = nullptr;
-	// end of temp code for nodes development
+#endif
+	CachedActionCount = 0;
 }
 
 void UDataprepAsset::PostLoad()
@@ -74,6 +75,7 @@ void UDataprepAsset::PostLoad()
 			bMarkDirty = true;
 		}
 
+#ifndef NO_BLUEPRINT
 		// Most likely a Dataprep asset from 4.23
 		if(StartNode == nullptr)
 		{
@@ -102,6 +104,7 @@ void UDataprepAsset::PostLoad()
 			UpdateActions();
 			bMarkDirty = true;
 		}
+#endif
 
 		if ( !Parameterization )
 		{
@@ -119,6 +122,8 @@ void UDataprepAsset::PostLoad()
 
 			GetOutermost()->SetDirtyFlag(true);
 		}
+
+		CachedActionCount = ActionAssets.Num();
 	}
 }
 
@@ -142,6 +147,15 @@ bool UDataprepAsset::Rename(const TCHAR* NewName/* =nullptr */, UObject* NewOute
 	return bWasRename;
 }
 
+void UDataprepAsset::PostEditUndo()
+{
+	UDataprepAssetInterface::PostEditUndo();
+
+	OnActionChanged.Broadcast(nullptr, (ActionAssets.Num() == CachedActionCount) ? FDataprepAssetChangeType::ActionMoved : FDataprepAssetChangeType::ActionRemoved);
+
+	CachedActionCount = ActionAssets.Num();
+}
+
 const UDataprepActionAsset* UDataprepAsset::GetAction(int32 Index) const
 {
 	if ( ActionAssets.IsValidIndex( Index ) )
@@ -162,6 +176,7 @@ const UDataprepActionAsset* UDataprepAsset::GetAction(int32 Index) const
 	return nullptr;
 }
 
+#ifndef NO_BLUEPRINT
 void UDataprepAsset::RemoveActionUsingBP(int32 Index)
 {
 	if ( ActionAssets.IsValidIndex( Index ) )
@@ -364,6 +379,7 @@ bool UDataprepAsset::CreateBlueprint()
 
 	return true;
 }
+#endif
 
 bool UDataprepAsset::CreateParameterization()
 {
@@ -448,6 +464,7 @@ void UDataprepAsset::GetExistingParameterNamesForType(FProperty* Property, bool
 	Parameterization->GetExistingParameterNamesForType( Property, bIsDescribingFullProperty, OutValidExistingNames, OutInvalidNames );
 }
 
+#ifndef NO_BLUEPRINT
 void UDataprepAsset::OnDataprepBlueprintChanged( UBlueprint* InBlueprint )
 {
 	if(InBlueprint == DataprepRecipeBP)
@@ -508,7 +525,10 @@ void UDataprepAsset::UpdateActions()
 			NextNodeInPin = NodeOutPin ? ( NodeOutPin->LinkedTo.Num() > 0 ? NodeOutPin->LinkedTo[0] : nullptr ) : nullptr;
 		}
 	}
+
+	CachedActionCount = ActionAssets.Num();
 }
+#endif
 
 int32 UDataprepAsset::AddAction(const UDataprepActionAsset* InAction)
 {
@@ -517,9 +537,13 @@ int32 UDataprepAsset::AddAction(const UDataprepActionAsset* InAction)
 		Modify();
 
 		UDataprepActionAsset* Action = DuplicateObject<UDataprepActionAsset>( InAction, this);
+		Action->SetLabel( InAction->GetLabel() );
+
 		ActionAssets.Add( Action );
 		OnActionChanged.Broadcast(Action, FDataprepAssetChangeType::ActionAdded);
 
+		CachedActionCount = ActionAssets.Num();
+
 		return ActionAssets.Num() - 1;
 	}
 
@@ -530,6 +554,57 @@ int32 UDataprepAsset::AddAction(const UDataprepActionAsset* InAction)
 	return INDEX_NONE;
 }
 
+int32 UDataprepAsset::AddActions(const TArray<const UDataprepActionStep*>& InActionSteps, bool bCreateOne)
+{
+	if ( InActionSteps.Num() > 0 && InActionSteps[0] != nullptr )
+	{
+		Modify();
+
+		int32 PreviousActionCount = ActionAssets.Num();
+
+		if(bCreateOne == true)
+		{
+			UDataprepActionAsset* Action = NewObject<UDataprepActionAsset>( this, UDataprepActionAsset::StaticClass(), NAME_None, RF_Transactional );
+			Action->SetLabel( TEXT("New Action") );
+
+			ActionAssets.Add( Action );
+
+			Action->AddSteps(InActionSteps);
+		}
+		else
+		{
+
+			for(const UDataprepActionStep* InActionStep : InActionSteps)
+			{
+				if(InActionStep)
+				{
+					UDataprepActionAsset* Action = NewObject<UDataprepActionAsset>( this, UDataprepActionAsset::StaticClass(), NAME_None, RF_Transactional );
+					Action->SetLabel( TEXT("New Action") );
+
+					ActionAssets.Add( Action );
+
+					Action->AddStep(InActionStep);
+				}
+			}
+		}
+
+		CachedActionCount = ActionAssets.Num();
+
+		if(PreviousActionCount != CachedActionCount)
+		{
+			OnActionChanged.Broadcast(ActionAssets.Last(), FDataprepAssetChangeType::ActionAdded);
+
+			return ActionAssets.Num() - 1;
+		}
+	}
+
+	UE_LOG( LogDataprepCore, Error, TEXT("UDataprepAsset::AddAction: None of the action steps is invalid") );
+	ensure(false);
+
+	// Invalid
+	return INDEX_NONE;
+}
+
 bool UDataprepAsset::InsertAction(const UDataprepActionAsset* InAction, int32 Index)
 {
 	if(!ActionAssets.IsValidIndex(Index))
@@ -543,10 +618,14 @@ bool UDataprepAsset::InsertAction(const UDataprepActionAsset* InAction, int32 In
 		Modify();
 
 		UDataprepActionAsset* Action = DuplicateObject<UDataprepActionAsset>( InAction, this);
+		Action->SetLabel( InAction->GetLabel());
+
 		ActionAssets.Insert( Action, Index );
 
 		OnActionChanged.Broadcast(Action, FDataprepAssetChangeType::ActionAdded);
 
+		CachedActionCount = ActionAssets.Num();
+
 		return true;
 	}
 
@@ -557,6 +636,59 @@ bool UDataprepAsset::InsertAction(const UDataprepActionAsset* InAction, int32 In
 	return false;
 }
 
+bool UDataprepAsset::InsertActions(const TArray<const UDataprepActionStep*>& InActionSteps, int32 Index, bool bCreateOne)
+{
+	if ( InActionSteps.Num() > 0 && InActionSteps[0] != nullptr )
+	{
+		Modify();
+
+		int32 PreviousActionCount = ActionAssets.Num();
+
+		if(bCreateOne == true)
+		{
+			UDataprepActionAsset* Action = NewObject<UDataprepActionAsset>( this, UDataprepActionAsset::StaticClass(), NAME_None, RF_Transactional );
+			Action->SetLabel( TEXT("New Action") );
+
+			ActionAssets.Insert( Action, Index );
+
+			Action->AddSteps(InActionSteps);
+		}
+		else
+		{
+			int32 InsertIndex = Index;
+
+			for(const UDataprepActionStep* InActionStep : InActionSteps)
+			{
+				if(InActionStep)
+				{
+					UDataprepActionAsset* Action = NewObject<UDataprepActionAsset>( this, UDataprepActionAsset::StaticClass(), NAME_None, RF_Transactional );
+					Action->SetLabel( TEXT("New Action") );
+
+					ActionAssets.Insert( Action, InsertIndex );
+					++InsertIndex;
+
+					Action->AddStep(InActionStep);
+				}
+			}
+		}
+
+		CachedActionCount = ActionAssets.Num();
+
+		if(CachedActionCount != PreviousActionCount)
+		{
+			OnActionChanged.Broadcast(ActionAssets.Last(), FDataprepAssetChangeType::ActionAdded);
+
+			return true;
+		}
+	}
+
+	UE_LOG( LogDataprepCore, Error, TEXT("UDataprepAsset::AddAction: None of the action steps is invalid") );
+	ensure(false);
+
+	// Invalid
+	return false;
+}
+
 bool UDataprepAsset::MoveAction(int32 SourceIndex, int32 DestinationIndex)
 {
 	if ( SourceIndex == DestinationIndex )
@@ -605,6 +737,8 @@ bool UDataprepAsset::RemoveAction(int32 Index)
 
 		ActionAssets.RemoveAt( Index );
 
+		CachedActionCount = ActionAssets.Num();
+
 		OnActionChanged.Broadcast(ActionAsset, FDataprepAssetChangeType::ActionRemoved);
 
 		return true;
@@ -654,6 +788,8 @@ bool UDataprepAsset::RemoveActions(const TArray<int32>& Indices)
 			}
 		}
 
+		CachedActionCount = ActionAssets.Num();
+
 		// Notify on last action removed
 		OnActionChanged.Broadcast(ActionAsset, FDataprepAssetChangeType::ActionRemoved);
 
diff --git a/Engine/Plugins/Enterprise/DataprepEditor/Source/DataprepCore/Public/DataprepActionAsset.h b/Engine/Plugins/Enterprise/DataprepEditor/Source/DataprepCore/Public/DataprepActionAsset.h
--- a/Engine/Plugins/Enterprise/DataprepEditor/Source/DataprepCore/Public/DataprepActionAsset.h
+++ b/Engine/Plugins/Enterprise/DataprepEditor/Source/DataprepCore/Public/DataprepActionAsset.h
@@ -207,6 +207,13 @@ public:
 	 */
 	int32 AddStep(const UDataprepActionStep* ActionStep);
 
+	/**
+	 * Add a copy of each step to the action
+	 * @param ActionSteps The array of steps we want to duplicate in the action
+	 * @return The index of the last added step or index none if no step is invalid
+	 */
+	int32 AddSteps(const TArray<const UDataprepActionStep*>& ActionSteps);
+
 	/**
 	 * Add a copy of the step object to the action
 	 * @param StepObject The step object we want to duplicate in the action
@@ -222,6 +229,14 @@ public:
 	 */
 	bool InsertStep(const UDataprepActionStep* ActionStep, int32 Index);
 
+	/**
+	 * Insert a copy of each step into the action at the requested index
+	 * @param Index Index to insert the action step at
+	 * @param ActionSteps The array of steps to duplicate in the action
+	 * @return The index of the added step or index none if the step is invalid
+	 */
+	bool InsertSteps(const TArray<const UDataprepActionStep*>& ActionSteps, int32 Index);
+
 	/**
 	 * Access to a step of the action
 	 * @param Index the index of the desired step
@@ -289,7 +304,7 @@ public:
 	bool bExecutionInterrupted;
 
 	/** Getter and Setter on the UI text of the action */
-	const TCHAR* GetLabel() { return *Label; }
+	const TCHAR* GetLabel() const { return *Label; }
 	void SetLabel( const TCHAR* InLabel ) { Label = InLabel ? InLabel : TEXT(""); }
 
 	/**
@@ -409,6 +424,7 @@ private:
 	bool bWorkingSetHasChanged;
 
 	/** UI label of the action */
+	UPROPERTY()
 	FString Label;
 
 	/** Package which static meshes will be added to */
diff --git a/Engine/Plugins/Enterprise/DataprepEditor/Source/DataprepCore/Public/DataprepAsset.h b/Engine/Plugins/Enterprise/DataprepEditor/Source/DataprepCore/Public/DataprepAsset.h
--- a/Engine/Plugins/Enterprise/DataprepEditor/Source/DataprepCore/Public/DataprepAsset.h
+++ b/Engine/Plugins/Enterprise/DataprepEditor/Source/DataprepCore/Public/DataprepAsset.h
@@ -10,8 +10,11 @@
 
 #include "DataprepAsset.generated.h"
 
+#ifndef NO_BLUEPRINT
 class UBlueprint;
+#endif
 class UDataprepActionAsset;
+class UDataprepActionStep;
 class UDataprepParameterizableObject;
 class UDataprepParameterization;
 class UDataprepProducers;
@@ -39,6 +42,7 @@ public:
 	// UObject interface
 	virtual void PostLoad() override;
 	virtual bool Rename(const TCHAR* NewName/* =nullptr */, UObject* NewOuter/* =nullptr */, ERenameFlags Flags/* =REN_None */) override;
+	virtual void PostEditUndo() override;
 	// End of UObject interface
 
 	// UDataprepAssetInterface interface
@@ -66,6 +70,15 @@ public:
 	 */
 	int32 AddAction(const UDataprepActionAsset* Action);
 
+	/**
+	 * Creates an action from the array of action steps or one action per action steps
+	 * then add the action(s) to the Dataprep asset
+	 * @param ActionSteps The array of action steps to process
+	 * @param bCreateOne Indicates if one or more action assets should be created. By default one is created
+	 * @return The index of the last added action or index none if the action is invalid
+	 */
+	int32 AddActions(const TArray<const UDataprepActionStep*>& ActionSteps, bool bCreateOne = true);
+
 	/**
 	 * Insert a copy of the action to the Dataprep asset at the requested index
 	 * @param Action The action we want to duplicate in the Dataprep asset
@@ -74,6 +87,15 @@ public:
 	 */
 	bool InsertAction(const UDataprepActionAsset* InAction, int32 Index);
 
+	/**
+	 * Creates an action from the array of action steps or one action per action steps
+	 * then insert the action(s) to the Dataprep asset at the requested index
+	 * @param ActionSteps The array of action steps to process
+	 * @param Index The index at which the insertion must happen
+	 * @return True if the insertion is successful, false if the action steps or the index are invalid
+	 */
+	bool InsertActions(const TArray<const UDataprepActionStep*>& InActionSteps, int32 Index, bool bCreateOne = true);
+
 	/**
 	 * Move an action to another spot in the order of actions
 	 * This operation take O(n) time. Where n is the absolute value of SourceIndex - DestinationIndex
@@ -106,6 +128,7 @@ public:
 
 	bool CreateParameterization();
 
+#ifndef NO_BLUEPRINT
 	// Temp code for the nodes development
 	bool CreateBlueprint();
 
@@ -154,9 +177,10 @@ public:
 	DECLARE_EVENT_OneParam(UDataprepAsset, FOnDataprepBlueprintChange, UObject* /*The object that was modified*/)
 	FOnDataprepBlueprintChange& GetOnBlueprintChanged() { return OnBlueprintChanged; }
 	// end of temp code for nodes development
+#endif
 
 public:
-	// Functions specific to the parametrization of the dataprep asset
+	// Functions specific to the parametrization of the Dataprep asset
 
 	/**
 	 * Event to notify the ui that a dataprep parametrization was modified
@@ -217,7 +241,11 @@ private:
 
 	FOnDataprepActionAssetChange OnActionChanged;
 
+	int32 CachedActionCount;
+
+#ifndef NO_BLUEPRINT
 	/** Event broadcasted when object in the pipeline was modified (Only broadcasted on changes that can affect the result of execution) */
 	FOnDataprepBlueprintChange OnBlueprintChanged;
 	// end of temp code for nodes development
+#endif
 };
\ No newline at end of file
diff --git a/Engine/Plugins/Enterprise/DataprepEditor/Source/DataprepEditor/Private/SchemaActions/DataprepDragDropOp.cpp b/Engine/Plugins/Enterprise/DataprepEditor/Source/DataprepEditor/Private/SchemaActions/DataprepDragDropOp.cpp
--- a/Engine/Plugins/Enterprise/DataprepEditor/Source/DataprepEditor/Private/SchemaActions/DataprepDragDropOp.cpp
+++ b/Engine/Plugins/Enterprise/DataprepEditor/Source/DataprepEditor/Private/SchemaActions/DataprepDragDropOp.cpp
@@ -4,6 +4,7 @@
 
 #include "BlueprintNodes/K2Node_DataprepAction.h"
 #include "DataprepActionAsset.h"
+#include "DataprepAsset.h"
 #include "DataprepEditorLogCategory.h"
 #include "DataprepGraph/DataprepGraph.h"
 #include "DataprepGraph/DataprepGraphActionNode.h"
@@ -31,9 +32,6 @@
 
 #define LOCTEXT_NAMESPACE "DataprepDragAndDrop"
 
-class SDataprepGraphActionStepNode;
-class SGraphPanel;
-
 FDataprepDragDropOp::FDataprepDragDropOp()
 	: FGraphEditorDragDropAction()
 	, HoveredDataprepActionContext()
@@ -108,6 +106,21 @@ void FDataprepDragDropOp::HoverTargetChanged()
 		bDropTargetValid = true;
 		DrapDropText = LOCTEXT("TargetIsDataprepActionContext", "Add a Step to Dataprep Action");
 	}
+	else if(Cast<UDataprepGraphActionStepNode>(GetHoveredNode()) != nullptr)
+	{
+		bDropTargetValid = true;
+		DrapDropText = LOCTEXT("CopyDataprepActionStepNode", "Add/Insert step to location");
+	}
+	else if(Cast<UDataprepGraphActionNode>(GetHoveredNode()) != nullptr)
+	{
+		bDropTargetValid = true;
+		DrapDropText = LOCTEXT("CopyDataprepActionAssetNode", "Add/Insert step to location");
+	}
+	else if(Cast<UDataprepGraphRecipeNode>(GetHoveredNode()) != nullptr)
+	{
+		bDropTargetValid = true;
+		DrapDropText = LOCTEXT("InsertDataprepActionAssetNode", "Insert action to location");
+	}
 	else if ( UEdGraph* EdGraph = GetHoveredGraph() )
 	{
 		if ( const UEdGraphSchema_K2* GraphSchema_k2 = Cast<UEdGraphSchema_K2>( EdGraph->GetSchema() ) )
@@ -194,13 +207,13 @@ EVisibility FDataprepDragDropOp::GetErrorIconVisible() const
 
 FReply FDataprepDragDropOp::DroppedOnNode(FVector2D ScreenPosition, FVector2D GraphPosition)
 {
-	if(Cast<UDataprepGraphActionStepNode>(GetHoveredNode()) != nullptr)
+	if(UDataprepGraphActionStepNode* ActionStepNode = Cast<UDataprepGraphActionStepNode>(GetHoveredNode()))
 	{
-		return DoDropOnActionStep(ScreenPosition, GraphPosition);
+		return DoDropOnActionStep(ActionStepNode);
 	}
-	else if(Cast<UDataprepGraphActionNode>(GetHoveredNode()) != nullptr)
+	else if(UDataprepGraphActionNode* ActionNode = Cast<UDataprepGraphActionNode>(GetHoveredNode()))
 	{
-		return DoDropOnActionAsset(ScreenPosition, GraphPosition);
+		return DoDropOnActionAsset(ActionNode);
 	}
 
 	return FReply::Unhandled();
@@ -305,6 +318,10 @@ FText FDataprepDragDropOp::GetMessageText()
 		{
 			LastMessageText = bCopyRequested ? LOCTEXT("CopyDataprepActionAssetNode", "Copy step to location") : LOCTEXT("MoveDataprepActionAssetNode", "Move step to location");
 		}
+		else if(Cast<UDataprepGraphRecipeNode>(GetHoveredNode()) != nullptr)
+		{
+			LastMessageText = LOCTEXT("InsertDataprepActionAssetNode", "Insert action to location");
+		}
 	}
 	else if(GetHoveredNode() != nullptr)
 	{
@@ -378,7 +395,7 @@ void FDataprepDragDropOp::DoDropOnPanel(const TSharedRef<SWidget>& Panel, FVecto
 	}
 }
 
-FReply FDataprepDragDropOp::DoDropOnActionStep(FVector2D ScreenPosition, FVector2D GraphPosition)
+FReply FDataprepDragDropOp::DoDropOnActionStep(UDataprepGraphActionStepNode* TargetActionStepNode)
 {
 	FModifierKeysState ModifierKeyState = FSlateApplication::Get().GetModifierKeys();
 	bool bCopyRequested = ModifierKeyState.IsControlDown() || ModifierKeyState.IsCommandDown();
@@ -388,128 +405,233 @@ FReply FDataprepDragDropOp::DoDropOnActionStep(FVector2D ScreenPosition, FVector
 		return FReply::Handled().EndDragDrop();
 	}
 
-	UDataprepGraphActionStepNode* TargetActionStepNode = Cast<UDataprepGraphActionStepNode>(GetHoveredNode());
-	ensure(TargetActionStepNode != nullptr);
 	UDataprepActionAsset* TargetActionAsset = TargetActionStepNode->GetDataprepActionAsset();
+	ensure(TargetActionAsset);
 
-	for(FDraggedStepEntry& DraggedStepEntry : DraggedSteps)
+	if(DraggedSteps.Num() > 0)
 	{
-		if(DraggedStepEntry.Get<0>().IsValid() || DraggedStepEntry.Get<2>().IsValid())
+		for(FDraggedStepEntry& DraggedStepEntry : DraggedSteps)
 		{
-			FScopedTransaction Transaction( LOCTEXT("DropOnActionStep", "Copy/Move Dataprep Action Step") );
-			bool bTransactionSuccessful = true;
-
-			// External drag and drop
-			if(!DraggedStepEntry.Get<0>().IsValid())
+			if(DraggedStepEntry.Get<0>().IsValid() || DraggedStepEntry.Get<2>().IsValid())
 			{
-				bTransactionSuccessful &= TargetActionAsset->InsertStep( DraggedStepEntry.Get<2>().Get(), TargetActionStepNode->GetStepIndex() );
-			}
-			// Drag and drop within an action asset or between two action assets
-			else
-			{
-				UDataprepActionAsset* SourceActionAsset = DraggedStepEntry.Get<0>().Get();
-				check(SourceActionAsset);
-
-				int32 StepIndex = DraggedStepEntry.Get<1>();
-				check(StepIndex != INDEX_NONE);
-
-				// Hold onto the action step in case of a move
-				TStrongObjectPtr<UDataprepActionStep> SourceActionStepPtr = TStrongObjectPtr<UDataprepActionStep>( SourceActionAsset->GetStep(StepIndex).Get() );
-				check(SourceActionStepPtr.IsValid());
+				FScopedTransaction Transaction( LOCTEXT("DropOnActionStep", "Copy/Move Dataprep Action Step") );
+					bool bTransactionSuccessful = true;
 
-				// source action asset differs from target action asset
-				if( TargetActionAsset != SourceActionAsset)
-				{
-					if(!bCopyRequested)
+					// External drag and drop
+					if(!DraggedStepEntry.Get<0>().IsValid())
+					{
+						bTransactionSuccessful &= TargetActionAsset->InsertStep( DraggedStepEntry.Get<2>().Get(), TargetActionStepNode->GetStepIndex() );
+					}
+					// Drag and drop within an action asset or between two action assets
+					else
 					{
-						bTransactionSuccessful &= SourceActionAsset->RemoveStep( StepIndex );
+						UDataprepActionAsset* SourceActionAsset = DraggedStepEntry.Get<0>().Get();
+						check(SourceActionAsset);
+
+						int32 StepIndex = DraggedStepEntry.Get<1>();
+						check(StepIndex != INDEX_NONE);
+
+						// Hold onto the action step in case of a move
+						TStrongObjectPtr<UDataprepActionStep> SourceActionStepPtr = TStrongObjectPtr<UDataprepActionStep>( SourceActionAsset->GetStep(StepIndex).Get() );
+						check(SourceActionStepPtr.IsValid());
+
+						// source action asset differs from target action asset
+						if( TargetActionAsset != SourceActionAsset)
+						{
+							if(!bCopyRequested)
+							{
+								bTransactionSuccessful &= SourceActionAsset->RemoveStep( StepIndex );
+							}
+
+							bTransactionSuccessful &= TargetActionAsset->InsertStep( SourceActionStepPtr.Get(), TargetActionStepNode->GetStepIndex() );
+						}
+						else if(bCopyRequested)
+						{
+							bTransactionSuccessful &= TargetActionAsset->InsertStep( SourceActionStepPtr.Get(), TargetActionStepNode->GetStepIndex() );
+						}
+						else
+						{
+							bTransactionSuccessful &= TargetActionAsset->MoveStep( StepIndex, TargetActionStepNode->GetStepIndex() );
+						}
 					}
 
-					bTransactionSuccessful &= TargetActionAsset->InsertStep( SourceActionStepPtr.Get(), TargetActionStepNode->GetStepIndex() );
-				}
-				else if(bCopyRequested)
-				{
-					bTransactionSuccessful &= TargetActionAsset->InsertStep( SourceActionStepPtr.Get(), TargetActionStepNode->GetStepIndex() );
-				}
-				else
+				if(!bTransactionSuccessful)
 				{
-					bTransactionSuccessful &= TargetActionAsset->MoveStep( StepIndex, TargetActionStepNode->GetStepIndex() );
+					Transaction.Cancel();
 				}
 			}
-
-			if(!bTransactionSuccessful)
-			{
-				Transaction.Cancel();
-			}
 		}
-	}
 
-	DraggedNodeWidgets.Reset();
-	DraggedSteps.Reset();
+		DraggedNodeWidgets.Reset();
+		DraggedSteps.Reset();
+	}
+	else
+	{
+		DropStepFromPanel(TargetActionAsset, TargetActionStepNode->GetStepIndex());
+	}
 
 	return FReply::Handled().EndDragDrop();
 }
 
-FReply FDataprepDragDropOp::DoDropOnActionAsset(FVector2D ScreenPosition, FVector2D GraphPosition)
+void FDataprepDragDropOp::DropStepFromPanel(UDataprepActionAsset* TargetActionAsset, int32 InsertIndex)
 {
-	FModifierKeysState ModifierKeyState = FSlateApplication::Get().GetModifierKeys();
-	bool bCopyRequested = ModifierKeyState.IsControlDown() || ModifierKeyState.IsCommandDown();
+	FScopedTransaction Transaction( LOCTEXT("AddStepToAction", "Add/Insert Dataprep Action Step") );
 
-	UDataprepGraphActionNode* TargetActionAssetNode = Cast<UDataprepGraphActionNode>(GetHoveredNode());
-	ensure(TargetActionAssetNode != nullptr);
-	UDataprepActionAsset* TargetActionAsset = TargetActionAssetNode->GetDataprepActionAsset();
+	FDataprepSchemaActionContext Context;
+	Context.DataprepActionPtr = TargetActionAsset;
+	Context.StepIndex = InsertIndex;
+
+	if ( !DataprepGraphOperation.Execute( Context ) )
+	{
+		Transaction.Cancel();
+	}
+}
 
-	for(FDraggedStepEntry& DraggedStepEntry : DraggedSteps)
+FReply FDataprepDragDropOp::DoDropOnTrack(UDataprepAsset* TargetDataprepAsset, int32 InsertIndex)
+{
+	ensure(Cast<UDataprepGraphRecipeNode>(GetHoveredNode()));
+
+	if(DraggedSteps.Num() > 0)
 	{
-		if(DraggedStepEntry.Get<0>().IsValid() || DraggedStepEntry.Get<2>().IsValid())
+		TArray<const UDataprepActionStep*> Steps;
+		Steps.Reserve(DraggedSteps.Num());
+
+		for(FDraggedStepEntry& DraggedStepEntry : DraggedSteps)
+		{
+			if(DraggedStepEntry.Get<2>().IsValid())
+			{
+				Steps.Add(DraggedStepEntry.Get<2>().Get());
+			}
+		}
+
+		if(Steps.Num() > 0)
 		{
-			FScopedTransaction Transaction( LOCTEXT("DropOnActionStep", "Copy/Move Dataprep Action Step") );
+			FModifierKeysState ModifierKeyState = FSlateApplication::Get().GetModifierKeys();
+			bool bOnlyOneRequested = ModifierKeyState.IsControlDown() || ModifierKeyState.IsCommandDown();
+
+			FScopedTransaction Transaction( LOCTEXT("DropOnTrack", "Add Dataprep Action") );
 			bool bTransactionSuccessful = true;
 
-			// External drag and drop
-			if(!DraggedStepEntry.Get<0>().IsValid())
+			if(InsertIndex >= 0 && InsertIndex < TargetDataprepAsset->GetActionCount())
 			{
-				bTransactionSuccessful &= TargetActionAsset->AddStep( DraggedStepEntry.Get<2>().Get() ) != INDEX_NONE;
+				TargetDataprepAsset->InsertActions(Steps, InsertIndex, bOnlyOneRequested);
 			}
-			// Drag and drop within an action asset or between two action assets
 			else
 			{
-				UDataprepActionAsset* SourceActionAsset = DraggedStepEntry.Get<0>().Get();
-				check(SourceActionAsset);
+				TargetDataprepAsset->AddActions(Steps, bOnlyOneRequested);
+			}
+		}
 
-				int32 StepIndex = DraggedStepEntry.Get<1>();
-				check(StepIndex != INDEX_NONE);
+		DraggedNodeWidgets.Reset();
+		DraggedSteps.Reset();
+	}
+	else
+	{
+		FScopedTransaction Transaction( LOCTEXT("AddNodeFromStep", "Add Dataprep Action Node from Step") );
 
-				// Hold onto the action step in case of a move
-				TStrongObjectPtr<UDataprepActionStep> SourceActionStepPtr = TStrongObjectPtr<UDataprepActionStep>( SourceActionAsset->GetStep(StepIndex).Get() );
-				check(SourceActionStepPtr.IsValid());
+		bool bTransactionSuccessful = false;
 
-				// source action asset differs from target action asset
-				if( TargetActionAsset != SourceActionAsset)
-				{
-					if(!bCopyRequested)
-					{
-						SourceActionAsset->RemoveStep( StepIndex );
-					}
+		UDataprepActionAsset* Action = Cast<UDataprepActionAsset>(UDataprepActionAsset::StaticClass()->GetDefaultObject());
 
-					TargetActionAsset->AddStep( SourceActionStepPtr.Get() );
-				}
-				else if(bCopyRequested)
+		if(InsertIndex >= 0 && InsertIndex < TargetDataprepAsset->GetActionCount())
+		{
+			bTransactionSuccessful = TargetDataprepAsset->InsertAction(Action, InsertIndex);
+		}
+		else
+		{
+			InsertIndex = TargetDataprepAsset->AddAction(Action);
+			bTransactionSuccessful = InsertIndex != INDEX_NONE;
+		}
+
+		if ( !bTransactionSuccessful )
+		{
+			Transaction.Cancel();
+			return FReply::Unhandled();
+		}
+
+		Action = TargetDataprepAsset->GetAction(InsertIndex);
+		Action->SetLabel(TEXT("New Action"));
+
+		FDataprepSchemaActionContext Context;
+		Context.DataprepActionPtr = Action;
+
+		bTransactionSuccessful = DataprepGraphOperation.Execute( Context );
+
+		if ( !bTransactionSuccessful )
+		{
+			TargetDataprepAsset->RemoveAction(InsertIndex);
+			Transaction.Cancel();
+		}
+	}
+
+	return FReply::Handled();
+}
+
+FReply FDataprepDragDropOp::DoDropOnActionAsset(UDataprepGraphActionNode* TargetActionAssetNode)
+{
+	FModifierKeysState ModifierKeyState = FSlateApplication::Get().GetModifierKeys();
+	bool bCopyRequested = ModifierKeyState.IsControlDown() || ModifierKeyState.IsCommandDown();
+
+	UDataprepActionAsset* TargetActionAsset = TargetActionAssetNode->GetDataprepActionAsset();
+
+	if(DraggedSteps.Num() > 0)
+	{
+		for(FDraggedStepEntry& DraggedStepEntry : DraggedSteps)
+		{
+			if(DraggedStepEntry.Get<0>().IsValid() || DraggedStepEntry.Get<2>().IsValid())
+			{
+				FScopedTransaction Transaction( LOCTEXT("DropOnActionStep", "Copy/Move Dataprep Action Step") );
+				bool bTransactionSuccessful = true;
+
+				// External drag and drop
+				if(!DraggedStepEntry.Get<0>().IsValid())
 				{
-					TargetActionAsset->AddStep( SourceActionStepPtr.Get() );
+					bTransactionSuccessful &= TargetActionAsset->AddStep( DraggedStepEntry.Get<2>().Get() ) != INDEX_NONE;
 				}
+				// Drag and drop within an action asset or between two action assets
 				else
 				{
-					TargetActionAsset->MoveStep( StepIndex, TargetActionAsset->GetStepsCount() - 1 );
+					UDataprepActionAsset* SourceActionAsset = DraggedStepEntry.Get<0>().Get();
+					check(SourceActionAsset);
+
+					int32 StepIndex = DraggedStepEntry.Get<1>();
+					check(StepIndex != INDEX_NONE);
+
+					// Hold onto the action step in case of a move
+					TStrongObjectPtr<UDataprepActionStep> SourceActionStepPtr = TStrongObjectPtr<UDataprepActionStep>( SourceActionAsset->GetStep(StepIndex).Get() );
+					check(SourceActionStepPtr.IsValid());
+
+					// source action asset differs from target action asset
+					if( TargetActionAsset != SourceActionAsset)
+					{
+						if(!bCopyRequested)
+						{
+							SourceActionAsset->RemoveStep( StepIndex );
+						}
+
+						TargetActionAsset->AddStep( SourceActionStepPtr.Get() );
+					}
+					else if(bCopyRequested)
+					{
+						TargetActionAsset->AddStep( SourceActionStepPtr.Get() );
+					}
+					else
+					{
+						TargetActionAsset->MoveStep( StepIndex, TargetActionAsset->GetStepsCount() - 1 );
+					}
 				}
-			}
 
-			if(!bTransactionSuccessful)
-			{
-				Transaction.Cancel();
+				if(!bTransactionSuccessful)
+				{
+					Transaction.Cancel();
+				}
 			}
 		}
 	}
+	else
+	{
+		DropStepFromPanel(TargetActionAsset);
+	}
 
 	DraggedNodeWidgets.Reset();
 	DraggedSteps.Reset();
diff --git a/Engine/Plugins/Enterprise/DataprepEditor/Source/DataprepEditor/Private/SchemaActions/DataprepDragDropOp.h b/Engine/Plugins/Enterprise/DataprepEditor/Source/DataprepEditor/Private/SchemaActions/DataprepDragDropOp.h
--- a/Engine/Plugins/Enterprise/DataprepEditor/Source/DataprepEditor/Private/SchemaActions/DataprepDragDropOp.h
+++ b/Engine/Plugins/Enterprise/DataprepEditor/Source/DataprepEditor/Private/SchemaActions/DataprepDragDropOp.h
@@ -13,6 +13,8 @@ struct FDataprepSchemaActionContext;
 class SDataprepGraphActionStepNode;
 class SGraphPanel;
 class UDataprepActionAsset;
+class UDataprepAsset;
+class UDataprepGraphActionNode;
 class UDataprepGraphActionStepNode;
 
 // Return true if there was a modification that require a transaction
@@ -50,6 +52,14 @@ public:
 	virtual EVisibility GetErrorIconVisible() const;
 	// End of FGraphEditorDragDropAction Interface
 
+	/**
+	 * Executes drop on track. A new action will inserted at the specified index.
+	 * @param InsertIndex Index at which to insert the new action. The action will be inserted
+	 *		  at the end if the index is not a valid one
+	 * @return FReply::handled.
+	 */
+	FReply DoDropOnTrack(UDataprepAsset* TargetDataprepAsset, int32 InsertIndex);
+
 	// Allow to add an extra step to the drag and drop before doing the dropping
 	void SetPreDropConfirmation(FDataprepPreDropConfirmation && Confirmation);
 
@@ -70,10 +80,10 @@ protected:
 	void DoDropOnPanel(const TSharedRef<SWidget>& Panel, FVector2D ScreenPosition, FVector2D GraphPosition, UEdGraph& Graph);
 
 	/** Executes drop on existing action step */
-	FReply DoDropOnActionStep(FVector2D ScreenPosition, FVector2D GraphPosition);
+	FReply DoDropOnActionStep(UDataprepGraphActionStepNode* TargetActionStepNode);
 
 	/** Executes drop on existing action step */
-	FReply DoDropOnActionAsset(FVector2D ScreenPosition, FVector2D GraphPosition);
+	FReply DoDropOnActionAsset(UDataprepGraphActionNode* TargetActionNode);
 
 	virtual void HoverTargetChangedWithNodes();
 
@@ -84,8 +94,17 @@ protected:
 
 private:
 	FText GetMessageText();
+
 	const FSlateBrush* GetIcon() const;
 
+	/**
+	 * Drop a step from the Operation panel to an action. The step will be added or inserted.
+	 * @param TargetActionAsset Action on which the insertion will be performed
+	 * @param InsertIndex Index at which the insertion must occur in the existing list of steps.
+	 *                    An index of -1 will trigger an addition.
+	 */
+	void DropStepFromPanel(UDataprepActionAsset* TargetActionAsset, int32 InsertIndex = INDEX_NONE);
+
 private:
 	typedef TTuple<TWeakObjectPtr<UDataprepActionAsset>,int32,TWeakObjectPtr<UDataprepActionStep>> FDraggedStepEntry;
 
diff --git a/Engine/Plugins/Enterprise/DataprepEditor/Source/DataprepEditor/Private/Widgets/DataprepGraph/DragDropActionNodeImpl.cpp b/Engine/Plugins/Enterprise/DataprepEditor/Source/DataprepEditor/Private/Widgets/DataprepGraph/DragDropActionNodeImpl.cpp
new file mode 100644
--- /dev/null
+++ b/Engine/Plugins/Enterprise/DataprepEditor/Source/DataprepEditor/Private/Widgets/DataprepGraph/DragDropActionNodeImpl.cpp
@@ -0,0 +1,122 @@
+// Copyright 1998-2019 Epic Games, Inc. All Rights Reserved.
+
+#include "Widgets/DataprepGraph/SDataprepGraphTrackNode.h"
+
+#include "DataprepAsset.h"
+#include "DataprepEditorLogCategory.h"
+#include "DataprepEditorStyle.h"
+#include "DataprepGraph/DataprepGraph.h"
+#include "DataprepGraph/DataprepGraphActionNode.h"
+
+#include "ScopedTransaction.h"
+
+#include "Framework/Application/SlateApplication.h"
+#include "Framework/Application/SlateUser.h"
+#include "Layout/Children.h"
+#include "NodeFactory.h"
+#include "SGraphPanel.h"
+#include "Widgets/Colors/SColorBlock.h"
+#include "Widgets/DataprepGraph/SDataprepGraphActionNode.h"
+#include "Widgets/DataprepGraph/SDataprepGraphEditor.h"
+#include "Widgets/DataprepWidgets.h"
+#include "Widgets/Layout/SBox.h"
+#include "Widgets/Layout/SConstraintCanvas.h"
+#include "Widgets/SBoxPanel.h"
+#include "Widgets/SOverlay.h"
+
+#define LOCTEXT_NAMESPACE "DataprepGraphEditor"
+
+class FDragDropActionNodeImpl : public FDragDropActionNode
+{
+public:
+	DRAG_DROP_OPERATOR_TYPE(FDragDropActionNodeImpl, FDragDropActionNode)
+
+	// FDragDropOperation interface
+	virtual void OnDrop( bool bDropWasHandled, const FPointerEvent& MouseEvent ) override;
+	virtual void OnDragged( const class FDragDropEvent& DragDropEvent ) override;
+	virtual FCursorReply OnCursorQuery() override { return FDragDropOperation::OnCursorQuery(); }
+	virtual TSharedPtr<SWidget> GetDefaultDecorator() const override { return FDragDropOperation::GetDefaultDecorator(); }
+	virtual FVector2D GetDecoratorPosition() const override { return FDragDropOperation::GetDecoratorPosition(); }
+	virtual void SetDecoratorVisibility(bool bVisible) override { FDragDropOperation::SetDecoratorVisibility(bVisible); }
+	virtual bool IsExternalOperation() const override { return FDragDropOperation::IsExternalOperation(); }
+	virtual bool IsWindowlessOperation() const override { return FDragDropOperation::IsWindowlessOperation(); }
+	// End of FDragDropOperation interface
+
+	TSharedPtr<SDataprepGraphTrackNode> TrackNodePtr;
+	TSharedPtr<SDataprepGraphActionNode> ActionNodePtr;
+};
+
+TSharedRef<FDragDropActionNode> FDragDropActionNode::New(const TSharedRef<SDataprepGraphTrackNode>& InTrackNodePtr, const TSharedRef<SDataprepGraphActionNode>& InDraggedNode)
+{
+	FDragDropActionNodeImpl* OperationImpl = new FDragDropActionNodeImpl;
+
+	OperationImpl->TrackNodePtr = InTrackNodePtr;
+	OperationImpl->ActionNodePtr = InDraggedNode;
+
+	OperationImpl->bCreateNewWindow = false;
+	OperationImpl->Construct();
+
+	InTrackNodePtr->OnStartNodeDrag(InDraggedNode);
+
+	TSharedRef<FDragDropActionNode> Operation = MakeShareable(new FDragDropActionNode);
+	Operation->Impl = TSharedPtr<FDragDropActionNode>(OperationImpl);
+
+	return Operation;
+}
+
+void FDragDropActionNodeImpl::OnDrop(bool bDropWasHandled, const FPointerEvent& MouseEvent)
+{
+	int32 NewExecutionOrder = TrackNodePtr->OnEndNodeDrag();
+
+	if(bDropWasHandled)
+	{
+		if(UDataprepAsset* DataprepAsset = TrackNodePtr->GetDataprepAsset())
+		{
+			FModifierKeysState ModifierKeyState = FSlateApplication::Get().GetModifierKeys();
+			const bool bCopyRequested = ModifierKeyState.IsControlDown() || ModifierKeyState.IsCommandDown();
+
+			FScopedTransaction Transaction( bCopyRequested ? LOCTEXT("OnDropCopy", "Add/Insert action") : LOCTEXT("OnDropMove", "Move action") );
+			bool bTransactionSuccessful = false;
+
+			if(bCopyRequested)
+			{
+				if(NewExecutionOrder >= DataprepAsset->GetActionCount())
+				{
+					bTransactionSuccessful = DataprepAsset->AddAction(ActionNodePtr->GetDataprepAction()) != INDEX_NONE;
+				}
+				else
+				{
+					bTransactionSuccessful = DataprepAsset->InsertAction(ActionNodePtr->GetDataprepAction(), NewExecutionOrder);
+				}
+			}
+			else if( NewExecutionOrder != ActionNodePtr->GetExecutionOrder())
+			{
+				bTransactionSuccessful = DataprepAsset->MoveAction(ActionNodePtr->GetExecutionOrder(), NewExecutionOrder);
+			}
+			else
+			{
+				TrackNodePtr->RefreshLayout();
+			}
+
+			if(!bTransactionSuccessful)
+			{
+				Transaction.Cancel();
+			}
+		}
+	}
+	else
+	{
+		TrackNodePtr->RefreshLayout();
+	}
+
+	FDragDropOperation::OnDrop(bDropWasHandled, MouseEvent);
+}
+
+void FDragDropActionNodeImpl::OnDragged(const FDragDropEvent& DragDropEvent)
+{
+	TrackNodePtr->OnNodeDragged( ActionNodePtr, DragDropEvent.GetScreenSpacePosition(), DragDropEvent.GetCursorDelta() );
+
+	FDragDropOperation::OnDragged(DragDropEvent);
+}
+
+#undef LOCTEXT_NAMESPACE
\ No newline at end of file
diff --git a/Engine/Plugins/Enterprise/DataprepEditor/Source/DataprepEditor/Private/Widgets/DataprepGraph/SDataprepGraphActionNode.cpp b/Engine/Plugins/Enterprise/DataprepEditor/Source/DataprepEditor/Private/Widgets/DataprepGraph/SDataprepGraphActionNode.cpp
--- a/Engine/Plugins/Enterprise/DataprepEditor/Source/DataprepEditor/Private/Widgets/DataprepGraph/SDataprepGraphActionNode.cpp
+++ b/Engine/Plugins/Enterprise/DataprepEditor/Source/DataprepEditor/Private/Widgets/DataprepGraph/SDataprepGraphActionNode.cpp
@@ -14,6 +14,7 @@
 // Engine Includes
 #include "SGraphPanel.h"
 #include "NodeFactory.h"
+#include "Widgets/Colors/SColorBlock.h"
 #include "Widgets/SBoxPanel.h"
 #include "Widgets/SOverlay.h"
 #include "Widgets/Layout/SBox.h"
@@ -34,7 +35,7 @@ public:
 	{
 		ParentPtr = InParent;
 
-		const float InterStepSpacing = 2.f;
+		const float InterStepSpacing = 5.f;
 
 		bIsHovered = false;
 
@@ -54,6 +55,8 @@ public:
 				]
 			]
 		];
+
+		bShowInsertionSlot = false;
 	}
 
 	// SWidget Interface
@@ -62,8 +65,10 @@ public:
 		TSharedPtr<FDataprepDragDropOp> DragActionStepNodeOp = DragDropEvent.GetOperationAs<FDataprepDragDropOp>();
 		if(DragActionStepNodeOp.IsValid() && ParentPtr.IsValid())
 		{
+			ParentTrackNodePtr.Pin()->OnDragLeave(DragDropEvent);
+
 			DragActionStepNodeOp->SetHoveredNode(ParentPtr.Pin()->GetNodeObj());
-			bIsHovered = true;
+			bShowInsertionSlot = true;
 		}
 
 		SCompoundWidget::OnDragEnter(MyGeometry, DragDropEvent);
@@ -75,7 +80,9 @@ public:
 		if(DragActionStepNodeOp.IsValid() && ParentPtr.IsValid())
 		{
 			DragActionStepNodeOp->SetHoveredNode(ParentPtr.Pin()->GetNodeObj());
-			bIsHovered = true;
+			bShowInsertionSlot = true;
+
+			return FReply::Handled();
 		}
 
 		return SCompoundWidget::OnDragOver(MyGeometry, DragDropEvent);
@@ -89,14 +96,14 @@ public:
 			DragActionStepNodeOp->SetHoveredNode(nullptr);
 		}
 
-		bIsHovered = false;
+		bShowInsertionSlot = false;
 
 		SCompoundWidget::OnDragLeave(DragDropEvent);
 	}
 
 	virtual FReply OnDrop(const FGeometry& MyGeometry, const FDragDropEvent& DragDropEvent) override
 	{
-		bIsHovered = false;
+		bShowInsertionSlot = false;
 	
 		// Process OnDrop if done by FDataprepDragDropOp
 		TSharedPtr<FDataprepDragDropOp> DragActionStepNodeOp = DragDropEvent.GetOperationAs<FDataprepDragDropOp>();
@@ -110,16 +117,27 @@ public:
 	}
 	// End of SWidget Interface
 
+	void SetParentTrackNode(TSharedPtr<SDataprepGraphTrackNode> InParentTrackNode)
+	{
+		ParentTrackNodePtr = InParentTrackNode;
+	}
+
+	void ShowInsertionSlot(bool bShow)
+	{
+		bShowInsertionSlot = bShow;
+	}
+
 private:
 	FSlateColor GetBorderBackgroundColor() const
 	{
 		static FLinearColor Hovered = FDataprepEditorStyle::GetColor("DataprepActionStep.DragAndDrop");
-		return bIsHovered ? Hovered : FLinearColor::Transparent;
+		return bShowInsertionSlot ? Hovered : FLinearColor::Transparent;
 	}
 
 private:
-	bool bIsHovered;
+	bool bShowInsertionSlot;
 	TWeakPtr<SDataprepGraphActionNode> ParentPtr;
+	TWeakPtr<SDataprepGraphTrackNode> ParentTrackNodePtr;
 };
 
 void SDataprepGraphActionNode::Construct(const FArguments& InArgs, UDataprepGraphActionNode* InActionNode)
@@ -147,6 +165,23 @@ void SDataprepGraphActionNode::MoveTo(const FVector2D& DesiredPosition, FNodeSet
 void SDataprepGraphActionNode::SetParentTrackNode(TSharedPtr<SDataprepGraphTrackNode> InParentTrackNode)
 {
 	ParentTrackNodePtr = InParentTrackNode;
+
+	// Update parent track on step widgets
+	for(TSharedPtr<SDataprepGraphActionStepNode>& ActionStepGraphNode : ActionStepGraphNodes)
+	{
+		ActionStepGraphNode->SetParentTrackNode(InParentTrackNode);
+	}
+
+	// Update parent on empty bottom widget
+	FChildren* StepListChildren = ActionStepListWidgetPtr->GetChildren();
+	TSharedRef<SDataprepEmptyActionStepNode> EmptyWidgetPtr = StaticCastSharedRef<SDataprepEmptyActionStepNode>(StepListChildren->GetChildAt(StepListChildren->Num() - 1));
+	EmptyWidgetPtr->SetParentTrackNode(InParentTrackNode);
+}
+
+void SDataprepGraphActionNode::UpdateExecutionOrder()
+{
+	ensure(Cast<UDataprepGraphActionNode>(GraphNode));
+	ExecutionOrder = Cast<UDataprepGraphActionNode>(GraphNode)->GetExecutionOrder();
 }
 
 TSharedRef<SWidget> SDataprepGraphActionNode::CreateNodeContentArea()
@@ -154,7 +189,27 @@ TSharedRef<SWidget> SDataprepGraphActionNode::CreateNodeContentArea()
 	if(DataprepActionPtr.IsValid())
 	{
 		PopulateActionStepListWidget();
-		return ActionStepListWidgetPtr.ToSharedRef();
+
+		return SNew(SVerticalBox)
+		+SVerticalBox::Slot()
+		.AutoHeight()
+		[
+			ActionStepListWidgetPtr.ToSharedRef()
+		]
+
+		//+SVerticalBox::Slot()
+		//.AutoHeight()
+		//[
+		//	SNew(SHorizontalBox)
+		//	+SHorizontalBox::Slot()
+		//	.AutoWidth()
+		//	[
+		//		SNew(SColorBlock)
+		//		.Color(FLinearColor::Transparent)
+		//		.Size( FVector2D(200.f, 10.f) )
+		//	]
+		//]
+		;
 	}
 
 	return 	SNew(STextBlock)
@@ -164,14 +219,15 @@ TSharedRef<SWidget> SDataprepGraphActionNode::CreateNodeContentArea()
 
 FReply SDataprepGraphActionNode::OnMouseButtonDown(const FGeometry& MyGeometry, const FPointerEvent& MouseEvent)
 {
-	FReply Reply = SGraphNode::OnMouseButtonDown( MyGeometry, MouseEvent);
+	GetOwnerPanel()->SelectionManager.ClickedOnNode(GraphNode, MouseEvent);
+	BorderBackgroundColor.Set(FDataprepEditorStyle::GetColor("DataprepActionStep.DragAndDrop"));
 
 	if( MouseEvent.GetEffectingButton() ==  EKeys::LeftMouseButton )
 	{
-		Reply = Reply.IsEventHandled() ? Reply.DetectDrag( AsShared(), EKeys::LeftMouseButton ) : FReply::Handled().DetectDrag( AsShared(), EKeys::LeftMouseButton );
+		return FReply::Handled().DetectDrag( AsShared(), EKeys::LeftMouseButton );
 	}
 
-	return Reply;
+	return FReply::Unhandled();
 }
 
 FCursorReply SDataprepGraphActionNode::OnCursorQuery( const FGeometry& MyGeometry, const FPointerEvent& CursorEvent ) const
@@ -184,27 +240,32 @@ FCursorReply SDataprepGraphActionNode::OnCursorQuery( const FGeometry& MyGeometr
 
 void SDataprepGraphActionNode::SetOwner(const TSharedRef<SGraphPanel>& OwnerPanel)
 {
-	SGraphNode::SetOwner(OwnerPanel);
-
-	for(TSharedPtr<SDataprepGraphActionStepNode>& ActionStepGraphNode : ActionStepGraphNodes)
+	if(!OwnerGraphPanelPtr.IsValid())
 	{
-		if (ActionStepGraphNode.IsValid())
+		SGraphNode::SetOwner(OwnerPanel);
+
+		for(TSharedPtr<SDataprepGraphActionStepNode>& ActionStepGraphNode : ActionStepGraphNodes)
 		{
-			ActionStepGraphNode->SetOwner(OwnerPanel);
-			OwnerPanel->AttachGraphEvents(ActionStepGraphNode);
+			if (ActionStepGraphNode.IsValid())
+			{
+				ActionStepGraphNode->SetOwner(OwnerPanel);
+				OwnerPanel->AttachGraphEvents(ActionStepGraphNode);
+			}
 		}
 	}
+	else
+	{
+		ensure(OwnerPanel == OwnerGraphPanelPtr);
+	}
 }
 
-FReply SDataprepGraphActionNode::OnDragDetected(const FGeometry & MyGeometry, const FPointerEvent & MouseEvent)
+FReply SDataprepGraphActionNode::OnDragDetected(const FGeometry& MyGeometry, const FPointerEvent& MouseEvent)
 {
 	if (UDataprepGraphActionNode* ActionNode = Cast<UDataprepGraphActionNode>(GraphNode))
 	{
 		if (UDataprepActionAsset* ActionAsset = ActionNode->GetDataprepActionAsset())
 		{
-			const TSharedRef<SDataprepGraphActionNode>& Node = SharedThis(this);
-			const TSharedRef<SDataprepGraphTrackNode>& Track = SharedThis(ParentTrackNodePtr.Pin().Get());
-			return FReply::Handled().BeginDragDrop(FDragGraphActionNode::New(Track, Node));
+			return FReply::Handled().BeginDragDrop(FDragDropActionNode::New(SharedThis(ParentTrackNodePtr.Pin().Get()), SharedThis(this)));
 		}
 	}
 
@@ -213,8 +274,9 @@ FReply SDataprepGraphActionNode::OnDragDetected(const FGeometry & MyGeometry, co
 
 FReply SDataprepGraphActionNode::OnDrop(const FGeometry& MyGeometry, const FDragDropEvent& DragDropEvent)
 {
-	TSharedPtr<FDragGraphActionNode> DragOperation = DragDropEvent.GetOperationAs<FDragGraphActionNode>();
 	SetCursor(EMouseCursor::Default);
+
+	TSharedPtr<FDragDropActionNode> DragOperation = DragDropEvent.GetOperationAs<FDragDropActionNode>();
 	if (DragOperation.IsValid())
 	{
 		return FReply::Handled().EndDragDrop();
@@ -223,6 +285,18 @@ FReply SDataprepGraphActionNode::OnDrop(const FGeometry& MyGeometry, const FDrag
 	return SGraphNode::OnDrop(MyGeometry, DragDropEvent);
 }
 
+void SDataprepGraphActionNode::OnDragEnter(const FGeometry& MyGeometry, const FDragDropEvent& DragDropEvent)
+{
+	// Track node is not notified of drag left, do it
+	TSharedPtr<FDragDropActionNode> DragOperation = DragDropEvent.GetOperationAs<FDragDropActionNode>();
+	if (ParentTrackNodePtr.IsValid() && DragOperation.IsValid())
+	{
+		ParentTrackNodePtr.Pin()->OnDragLeave(DragDropEvent);
+	}
+
+	SGraphNode::OnDragEnter(MyGeometry, DragDropEvent);
+}
+
 void SDataprepGraphActionNode::PopulateActionStepListWidget()
 {
 	if(!ActionStepListWidgetPtr.IsValid())
@@ -241,6 +315,8 @@ void SDataprepGraphActionNode::PopulateActionStepListWidget()
 	const int32 StepsCount = DataprepAction->GetStepsCount();
 	const UClass* GraphActionStepNodeClass = UDataprepGraphActionStepNode::StaticClass();
 
+	TSharedPtr<SDataprepGraphTrackNode> TrackNodePtr = ParentTrackNodePtr.Pin();
+
 	ActionStepGraphNodes.SetNum(StepsCount);
 
 	TSharedPtr<SGraphPanel> GraphPanelPtr = GetOwnerPanel();
@@ -259,6 +335,13 @@ void SDataprepGraphActionNode::PopulateActionStepListWidget()
 
 		TSharedPtr<SDataprepGraphActionStepNode> ActionStepGraphNode = StaticCastSharedPtr<SDataprepGraphActionStepNode>(FNodeFactory::CreateNodeWidget(ActionStepNode));
 
+		ActionStepGraphNode->SetParentNode(SharedThis(this));
+
+		if(TrackNodePtr.IsValid())
+		{
+			ActionStepGraphNode->SetParentTrackNode(TrackNodePtr);
+		}
+
 		ActionStepListWidgetPtr->AddSlot()
 		.AutoHeight()
 		[
@@ -276,11 +359,17 @@ void SDataprepGraphActionNode::PopulateActionStepListWidget()
 		}
 	}
 
+	TSharedPtr<SDataprepEmptyActionStepNode> BottomSlot;
 	ActionStepListWidgetPtr->AddSlot()
 	.AutoHeight()
 	[
-		SNew(SDataprepEmptyActionStepNode, SharedThis(this))
+		SAssignNew(BottomSlot, SDataprepEmptyActionStepNode, SharedThis(this))
 	];
+
+	if(TrackNodePtr.IsValid())
+	{
+		BottomSlot->SetParentTrackNode(TrackNodePtr);
+	}
 }
 
 void SDataprepGraphActionNode::OnStepsChanged()
diff --git a/Engine/Plugins/Enterprise/DataprepEditor/Source/DataprepEditor/Private/Widgets/DataprepGraph/SDataprepGraphActionNode.h b/Engine/Plugins/Enterprise/DataprepEditor/Source/DataprepEditor/Private/Widgets/DataprepGraph/SDataprepGraphActionNode.h
--- a/Engine/Plugins/Enterprise/DataprepEditor/Source/DataprepEditor/Private/Widgets/DataprepGraph/SDataprepGraphActionNode.h
+++ b/Engine/Plugins/Enterprise/DataprepEditor/Source/DataprepEditor/Private/Widgets/DataprepGraph/SDataprepGraphActionNode.h
@@ -29,6 +29,7 @@ public:
 	virtual FReply OnDragDetected(const FGeometry& MyGeometry, const FPointerEvent& MouseEvent) override;
 	virtual FCursorReply OnCursorQuery(const FGeometry& MyGeometry, const FPointerEvent& CursorEvent) const override;
 	virtual void SetOwner(const TSharedRef<SGraphPanel>& OwnerPanel) override;
+	virtual void OnDragEnter(const FGeometry& MyGeometry, const FDragDropEvent& DragDropEvent) override;
 	// End of SWidget interface
 
 	// SGraphNode interface
@@ -40,10 +41,15 @@ public:
 	void SetParentTrackNode(TSharedPtr<SDataprepGraphTrackNode> InParentTrackNode);
 
 	int32 GetExecutionOrder() const { return ExecutionOrder; }
-	void SetExecutionOrder( int32 InExecutionOrder ) { ExecutionOrder = InExecutionOrder; }
+	void UpdateExecutionOrder();
+
+	const UDataprepActionAsset* GetDataprepAction() const { return DataprepActionPtr.Get(); }
 
 private:
+	/** Callback to handle changes on array of steps in action */
 	void OnStepsChanged();
+
+	/** Reconstructs the list of widgets associated with the action's steps */
 	void PopulateActionStepListWidget();
 
 private:
diff --git a/Engine/Plugins/Enterprise/DataprepEditor/Source/DataprepEditor/Private/Widgets/DataprepGraph/SDataprepGraphActionStepNode.cpp b/Engine/Plugins/Enterprise/DataprepEditor/Source/DataprepEditor/Private/Widgets/DataprepGraph/SDataprepGraphActionStepNode.cpp
--- a/Engine/Plugins/Enterprise/DataprepEditor/Source/DataprepEditor/Private/Widgets/DataprepGraph/SDataprepGraphActionStepNode.cpp
+++ b/Engine/Plugins/Enterprise/DataprepEditor/Source/DataprepEditor/Private/Widgets/DataprepGraph/SDataprepGraphActionStepNode.cpp
@@ -36,7 +36,7 @@
 void SDataprepGraphActionStepNode::Construct(const FArguments& InArgs, UDataprepGraphActionStepNode* InActionStepNode)
 {
 	StepIndex = InActionStepNode->GetStepIndex();
-	bIsHovered = false;
+	bShowInsertionSlot = false;
 
 	GraphNode = InActionStepNode;
 
@@ -44,6 +44,11 @@ void SDataprepGraphActionStepNode::Construct(const FArguments& InArgs, UDataprep
 	UpdateGraphNode();
 }
 
+void SDataprepGraphActionStepNode::SetParentTrackNode(TSharedPtr<SDataprepGraphTrackNode> InParentTrackNode)
+{
+	ParentTrackNodePtr = InParentTrackNode;
+}
+
 TSharedPtr<SWidget> SDataprepGraphActionStepNode::GetStepTitleWidget() const
 {
 	if(SDataprepActionBlock* ActionStepBlock = ActionStepBlockPtr.Get())
@@ -148,12 +153,12 @@ FSlateColor SDataprepGraphActionStepNode::GetDragAndDropColor() const
 	static const FLinearColor BackgroundColor = FDataprepEditorStyle::GetColor("DataprepActionStep.BackgroundColor");
 	static const FLinearColor DragAndDrop = FDataprepEditorStyle::GetColor("DataprepActionStep.DragAndDrop");
 
-	return bIsHovered ? DragAndDrop : BackgroundColor;
+	return bShowInsertionSlot ? DragAndDrop : BackgroundColor;
 }
 
 FSlateColor SDataprepGraphActionStepNode::GetBorderBackgroundColor() const
 {
-	static const FLinearColor Selected(1.00f, 0.08f, 0.08f);
+	static const FLinearColor Selected = FDataprepEditorStyle::GetColor( "DataprepActionStep.DragAndDrop" );
 	static const FLinearColor BackgroundColor = FDataprepEditorStyle::GetColor("DataprepActionStep.BackgroundColor");
 
 	const bool bIsSelected = GetOwnerPanel()->SelectionManager.SelectedNodes.Contains(GraphNode);
@@ -175,9 +180,9 @@ FReply SDataprepGraphActionStepNode::OnMouseButtonDown(const FGeometry& MyGeomet
 
 FReply SDataprepGraphActionStepNode::OnMouseMove(const FGeometry& SenderGeometry, const FPointerEvent& MouseEvent)
 {
-	if (!MouseEvent.IsMouseButtonDown(EKeys::LeftMouseButton) && bIsHovered)
+	if (!MouseEvent.IsMouseButtonDown(EKeys::LeftMouseButton) && bShowInsertionSlot)
 	{
-		bIsHovered = false;
+		bShowInsertionSlot = false;
 	}
 
 	return FReply::Unhandled();
@@ -202,9 +207,11 @@ void SDataprepGraphActionStepNode::OnDragEnter(const FGeometry& MyGeometry, cons
 	TSharedPtr<FDataprepDragDropOp> DragNodeOp = DragDropEvent.GetOperationAs<FDataprepDragDropOp>();
 	if (DragNodeOp.IsValid())
 	{
+		ParentTrackNodePtr.Pin()->OnDragLeave(DragDropEvent);
+
 		// Inform the Drag and Drop operation that we are hovering over this node.
 		DragNodeOp->SetHoveredNode(GraphNode);
-		bIsHovered = DragNodeOp->IsValidDrop();
+		bShowInsertionSlot = DragNodeOp->IsValidDrop();
 
 		return;
 	}
@@ -220,7 +227,7 @@ FReply SDataprepGraphActionStepNode::OnDragOver(const FGeometry& MyGeometry, con
 	{
 		// Inform the Drag and Drop operation that we are hovering over this node.
 		DragNodeOp->SetHoveredNode(GraphNode);
-		bIsHovered = DragNodeOp->IsValidDrop();
+		bShowInsertionSlot = DragNodeOp->IsValidDrop();
 
 		return FReply::Handled();
 	}
@@ -235,7 +242,7 @@ void SDataprepGraphActionStepNode::OnDragLeave(const FDragDropEvent& DragDropEve
 	{
 		// Inform the Drag and Drop operation that we are not this widget anymore
 		DragNodeOp->SetHoveredNode(nullptr);
-		bIsHovered = false;
+		bShowInsertionSlot = false;
 
 		return;
 	}
@@ -245,7 +252,7 @@ void SDataprepGraphActionStepNode::OnDragLeave(const FDragDropEvent& DragDropEve
 
 FReply SDataprepGraphActionStepNode::OnDrop(const FGeometry& MyGeometry, const FDragDropEvent& DragDropEvent)
 {
-	bIsHovered = false;
+	bShowInsertionSlot = false;
 
 	// Process OnDrop if done by FDataprepDragDropOp
 	TSharedPtr<FDataprepDragDropOp> DragActionStepNodeOp = DragDropEvent.GetOperationAs<FDataprepDragDropOp>();
diff --git a/Engine/Plugins/Enterprise/DataprepEditor/Source/DataprepEditor/Private/Widgets/DataprepGraph/SDataprepGraphActionStepNode.h b/Engine/Plugins/Enterprise/DataprepEditor/Source/DataprepEditor/Private/Widgets/DataprepGraph/SDataprepGraphActionStepNode.h
--- a/Engine/Plugins/Enterprise/DataprepEditor/Source/DataprepEditor/Private/Widgets/DataprepGraph/SDataprepGraphActionStepNode.h
+++ b/Engine/Plugins/Enterprise/DataprepEditor/Source/DataprepEditor/Private/Widgets/DataprepGraph/SDataprepGraphActionStepNode.h
@@ -8,6 +8,7 @@
 #include "Widgets/DeclarativeSyntaxSupport.h"
 
 class SDataprepActionBlock;
+class SDataprepGraphActionNode;
 class SDataprepGraphTrackNode;
 class STextBlock;
 class UDataprepActionAsset;
@@ -46,6 +47,17 @@ public:
 	/** Returns a pointer to a widget displaying only the title of the action step */
 	virtual TSharedPtr<SWidget> GetStepTitleWidget() const;
 
+	void SetParentTrackNode(TSharedPtr<SDataprepGraphTrackNode> InParentTrackNode);
+	void SetParentNode(TSharedPtr<SDataprepGraphActionNode> InParentNode)
+	{
+		ParentNodePtr = InParentNode;
+	}
+
+	void ShowInsertionSlot(bool bShow)
+	{
+		bShowInsertionSlot = bShow;
+	}
+
 private:
 	/**
 	 * Returns a color depending whether the action step is selected or not
@@ -64,5 +76,11 @@ private:
 	int32 StepIndex;
 
 	/** Indicates if the widget is hovered during a drag and drop operation */
-	bool bIsHovered;
+	bool bShowInsertionSlot;
+
+	/** Pointer to the SDataprepGraphTrackNode displayed in the graph editor  */
+	TWeakPtr<SDataprepGraphTrackNode> ParentTrackNodePtr;
+
+	/** Pointer to the SDataprepGraphTrackNode displayed in the graph editor  */
+	TWeakPtr<SDataprepGraphActionNode> ParentNodePtr;
 };
diff --git a/Engine/Plugins/Enterprise/DataprepEditor/Source/DataprepEditor/Private/Widgets/DataprepGraph/SDataprepGraphEditor.cpp b/Engine/Plugins/Enterprise/DataprepEditor/Source/DataprepEditor/Private/Widgets/DataprepGraph/SDataprepGraphEditor.cpp
--- a/Engine/Plugins/Enterprise/DataprepEditor/Source/DataprepEditor/Private/Widgets/DataprepGraph/SDataprepGraphEditor.cpp
+++ b/Engine/Plugins/Enterprise/DataprepEditor/Source/DataprepEditor/Private/Widgets/DataprepGraph/SDataprepGraphEditor.cpp
@@ -95,10 +95,11 @@ void SDataprepGraphEditor::Construct(const FArguments& InArgs, UDataprepAsset* I
 	bMustRearrange = false;
 
 	LastLocalSize = FVector2D::ZeroVector;
-	LastLocation = FVector2D( BIG_NUMBER );
+	LastLocation = FVector2D( 0.f, -TopPadding );
 	LastZoomAmount = 1.f;
 
-	TrackGraphSize = FVector2D(SDataprepGraphTrackNode::NodeDesiredWidth);
+	FModifierKeysState ModifierKeyState = FSlateApplication::Get().GetModifierKeys();
+	bCachedControlKeyDown = ModifierKeyState.IsControlDown() || ModifierKeyState.IsCommandDown();
 }
 
 // #ueent_toremove: Temp code for the nodes development
@@ -109,6 +110,10 @@ void SDataprepGraphEditor::OnPipelineChanged(UBlueprint* InBlueprint)
 		TrackGraphNodePtr.Reset();
 		bIsComplete = false;
 		NotifyGraphChanged();
+
+		LastLocalSize = FVector2D::ZeroVector;
+		//LastLocation = FVector2D( BIG_NUMBER );
+		LastZoomAmount = 1.f;
 	}
 }
 
@@ -122,6 +127,10 @@ void SDataprepGraphEditor::OnDataprepAssetActionChanged(UObject* InObject, FData
 			TrackGraphNodePtr.Reset();
 			bIsComplete = false;
 			NotifyGraphChanged();
+
+			LastLocalSize = FVector2D::ZeroVector;
+			LastLocation = FVector2D::ZeroVector;
+			LastZoomAmount = 1.f;
 			break;
 		}
 
@@ -129,7 +138,7 @@ void SDataprepGraphEditor::OnDataprepAssetActionChanged(UObject* InObject, FData
 		{
 			if(SDataprepGraphTrackNode* TrackGraphNode = TrackGraphNodePtr.Pin().Get())
 			{
-				TrackGraphNode->ReArrangeActionNodes();
+				TrackGraphNode->OnActionsOrderChanged();
 			}
 			break;
 		}
@@ -142,21 +151,29 @@ void SDataprepGraphEditor::CacheDesiredSize(float InLayoutScaleMultiplier)
 
 	if(!bIsComplete && !NeedsPrepass())
 	{
-		// Get track SGraphNode and initialize it
-		if(UDataprepAsset* DataprepAsset = DataprepAssetPtr.Get())
+		if(!TrackGraphNodePtr.IsValid())
 		{
-			for(UEdGraphNode* EdGraphNode : GetCurrentGraph()->Nodes)
+			// Get track SGraphNode and initialize it
+			if(UDataprepAsset* DataprepAsset = DataprepAssetPtr.Get())
 			{
-				if(UDataprepGraphRecipeNode* TrackNode = Cast<UDataprepGraphRecipeNode>(EdGraphNode))
+				for(UEdGraphNode* EdGraphNode : GetCurrentGraph()->Nodes)
 				{
-					TrackGraphNodePtr = StaticCastSharedPtr<SDataprepGraphTrackNode>(TrackNode->GetWidget());
-					break;
+					if(UDataprepGraphRecipeNode* TrackNode = Cast<UDataprepGraphRecipeNode>(EdGraphNode))
+					{
+						TrackGraphNodePtr = StaticCastSharedPtr<SDataprepGraphTrackNode>(TrackNode->GetWidget());
+						break;
+					}
 				}
 			}
 		}
 
-		bMustRearrange = true;
-		bIsComplete = true;
+		if(TrackGraphNodePtr.IsValid())
+		{
+			bIsComplete = TrackGraphNodePtr.Pin()->RefreshLayout();
+			bMustRearrange = true;
+			// Force a change of viewpoint to update the canvas.
+			SetViewLocation(FVector2D(0.f, -TopPadding), 1.f);
+		}
 	}
 }
 
@@ -164,12 +181,12 @@ void SDataprepGraphEditor::UpdateBoundaries(const FVector2D& LocalSize, float Zo
 {
 	if(SDataprepGraphTrackNode* TrackGraphNode = TrackGraphNodePtr.Pin().Get())
 	{
-		TrackGraphSize = TrackGraphNode->Update(LocalSize, ZoomAmount);
+		CachedTrackNodeSize = TrackGraphNode->Update(LocalSize, ZoomAmount);
 	}
 
 	ViewLocationRangeOnY.Set( -TopPadding, -TopPadding );
 
-	const float DesiredVisualHeight = TopPadding + TrackGraphSize.Y + BottomPadding;
+	const float DesiredVisualHeight = CachedTrackNodeSize.Y * ZoomAmount;
 	if(LocalSize.Y < DesiredVisualHeight)
 	{
 		ViewLocationRangeOnY.Y = DesiredVisualHeight - LocalSize.Y;
@@ -182,6 +199,17 @@ void SDataprepGraphEditor::Tick(const FGeometry& AllottedGeometry, const double
 	// This happens after the first call to OnPaint on the editor
 	if(bIsComplete)
 	{
+		if(SDataprepGraphTrackNode* TrackGraphNode = TrackGraphNodePtr.Pin().Get())
+		{
+			FModifierKeysState ModifierKeyState = FSlateApplication::Get().GetModifierKeys();
+			bool bControlKeyDown = ModifierKeyState.IsControlDown() || ModifierKeyState.IsCommandDown();
+			if(bControlKeyDown != bCachedControlKeyDown)
+			{
+				bCachedControlKeyDown = bControlKeyDown;
+				TrackGraphNode->OnControlKeyChanged(bCachedControlKeyDown);
+			}
+		}
+
 		FVector2D Location;
 		float ZoomAmount = 1.f;
 		GetViewLocation( Location, ZoomAmount );
@@ -218,7 +246,7 @@ void SDataprepGraphEditor::UpdateLayout( const FVector2D& LocalSize, const FVect
 		if(Location.X != LastLocation.X)
 		{
 			const float ActualWidth = LocalSize.X / ZoomAmount;
-			const float MaxInX = TrackGraphSize.X > ActualWidth ? TrackGraphSize.X - ActualWidth : 0.f;
+			const float MaxInX = CachedTrackNodeSize.X > ActualWidth ? CachedTrackNodeSize.X - ActualWidth : 0.f;
 			ComputedLocation.X = Location.X < 0.f ? 0.f : Location.X >= MaxInX ? MaxInX : Location.X;
 		}
 
@@ -272,6 +300,17 @@ void SDataprepGraphEditor::OnDragEnter(const FGeometry & MyGeometry, const FDrag
 	SGraphEditor::OnDragEnter(MyGeometry, DragDropEvent);
 }
 
+FReply SDataprepGraphEditor::OnDragOver(const FGeometry & MyGeometry, const FDragDropEvent & DragDropEvent)
+{
+	TSharedPtr<FDataprepDragDropOp> DragNodeOp = DragDropEvent.GetOperationAs<FDataprepDragDropOp>();
+	if (TrackGraphNodePtr.IsValid() && DragNodeOp.IsValid())
+	{
+		TrackGraphNodePtr.Pin()->OnDragOver(MyGeometry, DragDropEvent);
+	}
+
+	return SGraphEditor::OnDragOver(MyGeometry, DragDropEvent);
+}
+
 void SDataprepGraphEditor::OnDragLeave(const FDragDropEvent & DragDropEvent)
 {
 	TSharedPtr<FDataprepDragDropOp> DragNodeOp = DragDropEvent.GetOperationAs<FDataprepDragDropOp>();
diff --git a/Engine/Plugins/Enterprise/DataprepEditor/Source/DataprepEditor/Private/Widgets/DataprepGraph/SDataprepGraphEditor.h b/Engine/Plugins/Enterprise/DataprepEditor/Source/DataprepEditor/Private/Widgets/DataprepGraph/SDataprepGraphEditor.h
--- a/Engine/Plugins/Enterprise/DataprepEditor/Source/DataprepEditor/Private/Widgets/DataprepGraph/SDataprepGraphEditor.h
+++ b/Engine/Plugins/Enterprise/DataprepEditor/Source/DataprepEditor/Private/Widgets/DataprepGraph/SDataprepGraphEditor.h
@@ -56,6 +56,7 @@ public:
 	virtual void CacheDesiredSize(float InLayoutScaleMultiplier) override;
 	virtual void Tick(const FGeometry& AllottedGeometry, const double InCurrentTime, const float InDeltaTime) override;
 	virtual void OnDragEnter(const FGeometry& MyGeometry, const FDragDropEvent& DragDropEvent) override;
+	virtual FReply OnDragOver(const FGeometry& MyGeometry, const FDragDropEvent& DragDropEvent) override;
 	virtual void OnDragLeave(const FDragDropEvent& DragDropEvent) override;
 	virtual FReply OnDrop(const FGeometry& MyGeometry, const FDragDropEvent& DragDropEvent) override;
 	// End of SWidget overrides
@@ -97,7 +98,7 @@ private:
 	FVector2D ViewLocationRangeOnY;
 
 	/** Size of graph being displayed */
-	mutable FVector2D TrackGraphSize;
+	mutable FVector2D CachedTrackNodeSize;
 
 	/** Size of graph being displayed */
 	TWeakObjectPtr<UDataprepAsset> DataprepAssetPtr;
@@ -105,6 +106,8 @@ private:
 	/** Size of graph being displayed */
 	mutable TWeakPtr<SDataprepGraphTrackNode> TrackGraphNodePtr;
 
+	bool bCachedControlKeyDown;
+
 	/** Padding used on the borders of the canvas */
 	// #ueent_wip: Will be moved to the Dataprep editor's style
 	static const float TopPadding;
