Copy up Dev-Enterprise @ CL11133089

diff --git a/Engine/Plugins/Enterprise/DataprepEditor/Source/DataprepEditor/Private/Widgets/DataprepGraph/SDataprepGraphTrackNode.cpp b/Engine/Plugins/Enterprise/DataprepEditor/Source/DataprepEditor/Private/Widgets/DataprepGraph/SDataprepGraphTrackNode.cpp
--- a/Engine/Plugins/Enterprise/DataprepEditor/Source/DataprepEditor/Private/Widgets/DataprepGraph/SDataprepGraphTrackNode.cpp
+++ b/Engine/Plugins/Enterprise/DataprepEditor/Source/DataprepEditor/Private/Widgets/DataprepGraph/SDataprepGraphTrackNode.cpp
@@ -7,6 +7,7 @@
 #include "DataprepEditorStyle.h"
 #include "DataprepGraph/DataprepGraph.h"
 #include "DataprepGraph/DataprepGraphActionNode.h"
+#include "SchemaActions/DataprepDragDropOp.h"
 
 #include "Layout/Children.h"
 #include "Widgets/DataprepGraph/SDataprepGraphActionNode.h"
@@ -25,9 +26,6 @@
 
 #define LOCTEXT_NAMESPACE "DataprepGraphEditor"
 
-float SDataprepGraphTrackNode::NodeDesiredWidth = 300.f;
-float SDataprepGraphTrackNode::NodeDesiredSpacing = 16.f;
-float SDataprepGraphTrackNode::TrackDesiredHeight = 40.f;
 FMargin SDataprepGraphTrackNode::NodePadding( 15.f, -5.f, 5.f, 10.f );
 
 class SDataprepGraphTrackWidget : public SHorizontalBox
@@ -36,84 +34,121 @@ public:
 	SLATE_BEGIN_ARGS(SDataprepGraphTrackWidget) {}
 	SLATE_END_ARGS();
 
-	void Construct(const FArguments& InArgs, TSharedPtr<SDataprepGraphTrackNode> InTrackNode)
+	void Construct(const FArguments& InArgs, TSharedPtr<SDataprepGraphTrackNode> InTrackNode);
+
+	void MoveDropSlotTo(int32 SlotIndex);
+
+	void SetActionSlot(int32 SlotIndex, const TSharedRef<SWidget>& Widget);
+
+	// SWidget Interface
+	virtual FVector2D ComputeDesiredSize(float InLayoutScaleMultiplier) const override;
+	// End of SWidget Interface
+
+	const FVector2D& GetWorkingArea() const
 	{
-		TrackNode = InTrackNode;
+		return WorkingArea;
+	}
 
-		SHorizontalBox::FSlot* FirstSlot = nullptr;
-		SHorizontalBox::FSlot* LastSlot = nullptr;
+	FVector2D GetWorkingArea()
+	{
+		return WorkingArea;
+	}
 
-		SHorizontalBox::Construct( SHorizontalBox::FArguments()
-			+SHorizontalBox::Slot()
-			.AutoWidth()
-			[
-				SNew(SColorBlock)
-				.Color( FDataprepEditorStyle::GetColor( "Graph.TrackEnds.BackgroundColor" ) )
-				.Size( TAttribute<FVector2D>::Create(TAttribute<FVector2D>::FGetter::CreateSP( InTrackNode.Get(), &SDataprepGraphTrackNode::GetLeftBlockSize ) ) )
-			]
-			+SHorizontalBox::Slot()
-			.AutoWidth()
-			[
-				SAssignNew( InnerCanvas, SConstraintCanvas )
+	FVector2D GetTrackArea()
+	{
+		return FVector2D(WorkingArea.X, TrackDesiredHeight);
+	}
 
-				// The outline. This is done by a background image
-				+ SConstraintCanvas::Slot()
-				.Anchors( FAnchors( 0.f, 0.f, 1.f, 1.f ) )
-				.Offset( FMargin() )
-				[
-					SNew(SColorBlock)
-					.Color( FDataprepEditorStyle::GetColor( "Graph.TrackInner.BackgroundColor" ) )
-					.Size( TAttribute<FVector2D>::Create(TAttribute<FVector2D>::FGetter::CreateSP( InTrackNode.Get(), &SDataprepGraphTrackNode::GetInnerBlockSize ) ) )
-				]
-			]
-			+SHorizontalBox::Slot()
-			.AutoWidth()
-			[
-				SNew(SColorBlock)
-				.Color( FDataprepEditorStyle::GetColor( "Graph.TrackEnds.BackgroundColor" ) )
-				.Size( TAttribute<FVector2D>::Create(TAttribute<FVector2D>::FGetter::CreateSP( InTrackNode.Get(), &SDataprepGraphTrackNode::GetRightBlockSize ) ) )
-			]
-		);
+	const FVector2D& GetCanvasArea() const
+	{
+		return CanvasArea;
 	}
 
-	// SWidget Interface
-	virtual FVector2D ComputeDesiredSize(float) const override
+	FVector2D GetCanvasArea()
 	{
-		if(SDataprepGraphTrackNode* TrackNodeRaw = TrackNode.Pin().Get())
-		{
-			const float TrackNodeWidth = TrackNodeRaw->GetLeftBlockSize().X + TrackNodeRaw->GetInnerBlockSize().X + TrackNodeRaw->GetRightBlockSize().X;
-			return FVector2D( TrackNodeWidth, SDataprepGraphTrackNode::TrackDesiredHeight );
-		}
-		
-		return FVector2D::ZeroVector;
+		return CanvasArea;
 	}
 
-	virtual int32 OnPaint(const FPaintArgs& Args, const FGeometry& AllottedGeometry, const FSlateRect& MyCullingRect, FSlateWindowElementList& OutDrawElements, int32 LayerId, const FWidgetStyle& InWidgetStyle, bool bParentEnabled) const override
+	FVector2D GetActioNodeAbscissaRange();
+
+	int32 GetHoveredActionNode(float Left, float Width);
+
+	float GetActionNodeAbscissa(int32 Index)
 	{
-		FArrangedChildren ArrangedChildren( EVisibility::Visible );
-		{
-			ArrangeChildren( AllottedGeometry, ArrangedChildren );
-		}
+		return ActionSlots.IsValidIndex(Index) ? ActionSlots[Index]->OffsetAttr.Get(FMargin()).Left : 0.f;
+	}
 
-		int32 MaxLayerId = LayerId;
-		for( int32 ChildIndex = 0, NodeIndex = 0; ChildIndex < ArrangedChildren.Num(); ++ChildIndex )
-		{
-			FArrangedWidget& CurWidget = ArrangedChildren[ ChildIndex ];
+	float ValidateNodeAbscissa(float InAbscissa);
 
-			const int32 CurWidgetsMaxLayerId = CurWidget.Widget->Paint(Args.WithNewParent(this), CurWidget.Geometry, MyCullingRect, OutDrawElements, LayerId, InWidgetStyle, ShouldBeEnabled(bParentEnabled));
-			MaxLayerId = FMath::Max(MaxLayerId, CurWidgetsMaxLayerId);
-		}
+	FVector2D GetNodeSize(const TSharedRef<SWidget>& Widget) const;
+
+	void RefreshLayout(SGraphPanel* GraphPanel);
+
+	void OnStartNodeDrag(int32 Index);
+
+	float OnNodeDragged(float Delta);
+
+	int32 OnEndNodeDrag();
+
+	bool IsDraggedNodeOnEnds();
+
+	/**
+	*/
+	void UpdateLayoutForCopy();
 
-		return MaxLayerId;
+	/**
+	*/
+	void UpdateLayoutForMove();
+
+	/** Toggles display of action nodes between copy vs move drop modes */
+	float OnControlKeyDown(bool bKeyDown);
+
+	void UpdateDragIndicator(FVector2D MousePosition);
+
+	int32 GetDragIndicatorIndex();
+
+	void ResetDragIndicator()
+	{
+		for(int32 Index = 0; Index < DropSlots.Num(); ++Index)
+		{
+			(*DropSlots[Index])[DropFiller.ToSharedRef()];
+		}
 	}
-	// End of SWidget Interface
 
-	TWeakPtr<SDataprepGraphTrackNode> TrackNode;
-	TSharedPtr<SConstraintCanvas> InnerCanvas;
+	void OnActionsOrderChanged(const TArray<TSharedPtr<SDataprepGraphActionNode>>& InActionNodes);
+
+	TWeakPtr<SDataprepGraphTrackNode> TrackNodePtr;
+	TSharedPtr<SConstraintCanvas> TrackCanvas;
+	TArray<TSharedRef<SWidget>> ActionNodes;
+	TArray<SConstraintCanvas::FSlot*> DropSlots;
+	TArray<SConstraintCanvas::FSlot*> ActionSlots;
+	SConstraintCanvas::FSlot* CanvasSlot;
+
+	TSharedPtr<SWidget> DropIndicator;
+	TSharedPtr<SColorBlock> DropFiller;
+	mutable FVector2D WorkingArea;
+	mutable FVector2D CanvasArea;
+	int32 LastDropSlot;
+
+	SConstraintCanvas::FSlot* DragSlot;
+	TArray<int32> NewActionsOrder;
+	int32 DragStartIndex;
+	int32 DragIndex;
+	FVector2D LastDragPosition;
+	FVector2D AbscissaRange;
+	bool bIsCopying;
+
+	static float SlotSpacing;
+	static float NodeDesiredWidth;
+	static float TrackDesiredHeight;
 
 	friend SDataprepGraphTrackNode;
 };
 
+float SDataprepGraphTrackWidget::SlotSpacing = 16.f;
+float SDataprepGraphTrackWidget::NodeDesiredWidth = 300.f;
+float SDataprepGraphTrackWidget::TrackDesiredHeight = 40.f;
+
 void SDataprepGraphTrackNode::Construct(const FArguments& InArgs, UDataprepGraphRecipeNode* InNode)
 {
 	bNodeDragging = false;
@@ -143,50 +178,10 @@ void SDataprepGraphTrackNode::UpdateGraphNode()
 	RightNodeBox.Reset();
 	LeftNodeBox.Reset();
 
-	TSharedPtr<SGraphPanel> GraphPanelPtr = OwnerGraphPanelPtr.Pin();
-
-	ContentScale.Bind( this, &SGraphNode::GetContentScale );
-
-	GetOrAddSlot( ENodeZone::Center )
-	.HAlign(HAlign_Fill)
-	.VAlign(VAlign_Center)
-	[
-		SNew(SBorder)
-		.BorderImage( FEditorStyle::GetBrush( "NoBorder" ) )
-		.Padding(0.f)
-		.BorderBackgroundColor( FLinearColor( 0.3f, 0.3f, 0.3f, 1.0f ) )
-		[
-			SNew(SVerticalBox)
-			+SVerticalBox::Slot()
-			.AutoHeight()
-			.Padding(5.f)
-			[
-				SNew(SBox)
-				//.MinDesiredHeight(NodeHeight)
-				[
-					SNew(SHorizontalBox)
-					+SHorizontalBox::Slot()
-					.AutoWidth()
-					[
-						SAssignNew(LeftNodeBox, SVerticalBox)
-					]
-					+SHorizontalBox::Slot()
-					.AutoWidth()
-					[
-						SAssignNew( TrackWidgetPtr, SDataprepGraphTrackWidget, SharedThis(this) )
-					]
-					+SHorizontalBox::Slot()
-					.AutoWidth()
-					[
-						SAssignNew(RightNodeBox, SVerticalBox)
-					]
-				]
-			]
-		]
-	];
-
 	if(UDataprepAsset* DataprepAsset = DataprepAssetPtr.Get())
 	{
+		TSharedPtr<SGraphPanel> GraphPanelPtr = OwnerGraphPanelPtr.Pin();
+
 		UEdGraph* EdGraph = GraphNode->GetGraph();
 
 		TSharedPtr<SDataprepGraphTrackNode> ParentTrackNodePtr = SharedThis(this);
@@ -208,13 +203,14 @@ void SDataprepGraphTrackNode::UpdateGraphNode()
 				NewActionNode->NodePosY = 0;
 
 				NewActionNode->Initialize(ActionAsset, Index);
-				
+
+				// #ueent_wip: Widget is created twice :-(
 				TSharedPtr< SDataprepGraphActionNode > ActionWidgetPtr = StaticCastSharedPtr<SDataprepGraphActionNode>(FNodeFactory::CreateNodeWidget(NewActionNode));
 				if(SDataprepGraphActionNode* ActionWidget = ActionWidgetPtr.Get())
 				{
 					if(GraphPanelPtr.IsValid())
 					{
-						GraphPanelPtr->AddGraphNode(ActionWidgetPtr.ToSharedRef());
+						ActionWidget->SetOwner(GraphPanelPtr.ToSharedRef());
 					}
 
 					ActionWidget->UpdateGraphNode();
@@ -226,130 +222,201 @@ void SDataprepGraphTrackNode::UpdateGraphNode()
 				}
 			}
 		}
+	}
 
-		NodeAbscissaMin = FMath::FloorToInt(LeftBlockSize.X + NodePadding.Left + InterNodeSpacing * 0.5f);
-		NodeAbscissaMax = NodeAbscissaMin + FMath::FloorToInt((float)(ActionNodes.Num()-1)*(NodeDesiredWidth + InterNodeSpacing));
+	ContentScale.Bind( this, &SGraphNode::GetContentScale );
 
-		// Position action nodes
-		ReArrangeActionNodes();
-	}
+	GetOrAddSlot( ENodeZone::Center )
+	.HAlign(HAlign_Fill)
+	.VAlign(VAlign_Center)
+	[
+		SNew(SBorder)
+		.BorderImage( FEditorStyle::GetBrush( "NoBorder" ) )
+		.Padding(0.f)
+		.BorderBackgroundColor( FLinearColor( 0.3f, 0.3f, 0.3f, 1.0f ) )
+		[
+			SNew(SVerticalBox)
+			+SVerticalBox::Slot()
+			.AutoHeight()
+			.Padding(5.f)
+			[
+				SNew(SBox)
+				[
+					SNew(SHorizontalBox)
+					+SHorizontalBox::Slot()
+					.AutoWidth()
+					[
+						SAssignNew( TrackWidgetPtr, SDataprepGraphTrackWidget, SharedThis(this) )
+					]
+				]
+			]
+		]
+	];
+
+	GraphNode->NodePosX = NodePadding.Left;
+	GraphNode->NodePosY = NodePadding.Top;
 }
 
-void SDataprepGraphTrackNode::ReArrangeActionNodes()
+bool SDataprepGraphTrackNode::CustomPrepass(float LayoutScaleMultiplier)
 {
-	SNodePanel::SNode::FNodeSet NodeFilter;
-	const float InitialOffset = LeftBlockSize.X + NodePadding.Left + InterNodeSpacing * 0.5f;
-	const float Increment = NodeDesiredWidth + InterNodeSpacing;
-	const float PositionY = NodePadding.Top;
+	return false;
+}
 
-	if(bNodeDragging)
+void SDataprepGraphTrackNode::OnActionsOrderChanged()
+{
+	if(!DataprepAssetPtr.IsValid())
 	{
-		FString SequenceText;
-		for(int32 Index = 0 ; Index < NewActionsOrder.Num(); ++Index)
-		{
-			const int32 OldExecutionOrder = NewActionsOrder[Index];
-			SequenceText += FString::Printf(TEXT("%d "), ActionNodes[OldExecutionOrder]->GetExecutionOrder());
+		return;
+	}
 
-			if(Index == CurrentOrder)
-			{
-				continue;
-			}
+	TMap<UDataprepActionAsset*, TSharedPtr< SDataprepGraphActionNode >> WidgetByAsset;
 
-			if( SDataprepGraphActionNode* ActionWidget = ActionNodes[OldExecutionOrder].Get())
-			{
-				ActionWidget->SGraphNode::MoveTo( FVector2D( NodeAbscissaMin + (float)Index * Increment, PositionY), NodeFilter);
-			}
+	for(TSharedPtr< SDataprepGraphActionNode >& ActionWidgetPtr : ActionNodes)
+	{
+
+		if( SDataprepGraphActionNode* ActionWidget = ActionWidgetPtr.Get() )
+		{
+			ensure(Cast<UDataprepGraphActionNode>(ActionWidgetPtr->GetNodeObj()) != nullptr);
+			WidgetByAsset.Emplace(Cast<UDataprepGraphActionNode>(ActionWidgetPtr->GetNodeObj())->GetDataprepActionAsset(), ActionWidgetPtr);
 		}
 	}
-	else
+
+	for(int32 Index = 0; Index < DataprepAssetPtr->GetActionCount(); ++Index)
 	{
-		for(TSharedPtr< SDataprepGraphActionNode >& ActionWidgetPtr : ActionNodes)
+		if(UDataprepActionAsset* ActionAsset = DataprepAssetPtr->GetAction(Index))
 		{
-			if( SDataprepGraphActionNode* ActionWidget = ActionWidgetPtr.Get() )
+			TSharedPtr< SDataprepGraphActionNode >* ActionWidgetPtr = WidgetByAsset.Find(ActionAsset);
+			if(ActionWidgetPtr != nullptr)
 			{
-				ActionWidget->SGraphNode::MoveTo( FVector2D( NodeAbscissaMin + (float)ActionWidget->GetExecutionOrder() * Increment, PositionY), NodeFilter);
-				ActionWidget->Invalidate(EInvalidateWidgetReason::RenderTransform);
+				UDataprepGraphActionNode* ActionNode = Cast<UDataprepGraphActionNode>((*ActionWidgetPtr)->GetNodeObj());
+				ensure(ActionNode != nullptr);
+
+				ActionNode->Initialize( ActionAsset, Index);
+				(*ActionWidgetPtr)->UpdateExecutionOrder();
 			}
 		}
 	}
-}
 
-void SDataprepGraphTrackNode::MoveTo(const FVector2D& NewPosition, FNodeSet& NodeFilter)
-{
-	SGraphNode::MoveTo( FVector2D::ZeroVector, NodeFilter);
+	// Reorder array according to new execution order
+	ActionNodes.Sort([](const TSharedPtr<SDataprepGraphActionNode> A, const TSharedPtr<SDataprepGraphActionNode> B) { return A->GetExecutionOrder() < B->GetExecutionOrder(); });
+
+	// Rearrange actions nodes along track
+	TrackWidgetPtr->OnActionsOrderChanged(ActionNodes);
 
-	NodeAbscissaMin = FMath::FloorToInt(LeftBlockSize.X + NodePadding.Left + InterNodeSpacing * 0.5f);
-	NodeAbscissaMax = NodeAbscissaMin + FMath::FloorToInt((float)(ActionNodes.Num()-1)*(NodeDesiredWidth + InterNodeSpacing));
+	// Invalidate graph panel for a redraw of all widgets
+	OwnerGraphPanelPtr.Pin()->Invalidate(EInvalidateWidgetReason::ChildOrder | EInvalidateWidgetReason::Layout);
 }
 
-void SDataprepGraphTrackNode::SetOwner(const TSharedRef<SGraphPanel>& OwnerPanel)
+bool SDataprepGraphTrackNode::RefreshLayout()
 {
-	ensure(!OwnerGraphPanelPtr.IsValid());
+	if(OwnerGraphPanelPtr.IsValid() && OwnerGraphPanelPtr.Pin()->GetAllChildren()->Num() > 0)
+	{
+		ensure(TrackWidgetPtr.IsValid());
 
-	SGraphNode::SetOwner(OwnerPanel);
+		FVector2D TrackWidgetOrigin = TrackWidgetPtr->GetCachedGeometry().LocalToAbsolute(FVector2D::ZeroVector);
+		FVector2D TrackNodeOrigin = GetCachedGeometry().LocalToAbsolute(FVector2D::ZeroVector);
 
-	for(TSharedPtr<SDataprepGraphActionNode>& ActionNodePtr : ActionNodes)
-	{
-		if (ActionNodePtr.IsValid())
-		{
-			OwnerPanel->AddGraphNode( ActionNodePtr.ToSharedRef() );
-		}
+		TrackWidgetOffset = TrackWidgetOrigin - TrackNodeOrigin;
+
+		TrackWidgetPtr->RefreshLayout(OwnerGraphPanelPtr.Pin().Get());
+
+		return true;
 	}
+
+	return false;
 }
 
-FVector2D SDataprepGraphTrackNode::Update(const FVector2D& LocalSize, float ZoomAmount)
+void SDataprepGraphTrackNode::MoveTo(const FVector2D& NewPosition, FNodeSet& NodeFilter)
+{
+	// Block track node to specific position
+	SGraphNode::MoveTo( FVector2D(NodePadding.Left, 0.f), NodeFilter);
+}
+
+const FSlateBrush* SDataprepGraphTrackNode::GetShadowBrush(bool bSelected) const
 {
-	// #ueent_wip: Find a way to avoid recompute of everything on each call
-	const float InvZoomAmount = 1.f / ZoomAmount;
+	return  FEditorStyle::GetNoBrush();
+}
 
-	if(NodeDesiredSpacing * ZoomAmount < 6.f)
+FReply SDataprepGraphTrackNode::OnDrop(const FGeometry& MyGeometry, const FDragDropEvent& DragDropEvent)
+{
+	SetCursor(EMouseCursor::Default);
+
+	TSharedPtr<FDragDropActionNode> DragOperation = DragDropEvent.GetOperationAs<FDragDropActionNode>();
+	if (DragOperation.IsValid())
 	{
-		InterNodeSpacing = FMath::CeilToFloat( NodeDesiredSpacing * InvZoomAmount );
+		return FReply::Handled().EndDragDrop();
 	}
-	else
+
+	TSharedPtr<FDataprepDragDropOp> DragActionStepNodeOp = DragDropEvent.GetOperationAs<FDataprepDragDropOp>();
+	if(DragActionStepNodeOp.IsValid())
 	{
-		InterNodeSpacing = NodeDesiredSpacing;
+		int32 InsertIndex = TrackWidgetPtr->GetDragIndicatorIndex();
+		TrackWidgetPtr->ResetDragIndicator();
+
+		DragActionStepNodeOp->DoDropOnTrack(DataprepAssetPtr.Get(), InsertIndex);
+
+		return FReply::Unhandled();
 	}
 
-	NodeAbscissaMin = FMath::FloorToInt(LeftBlockSize.X + NodePadding.Left + InterNodeSpacing * 0.5f);
-	NodeAbscissaMax = NodeAbscissaMin + FMath::FloorToInt((float)(ActionNodes.Num()-1)*(NodeDesiredWidth + InterNodeSpacing));
+	return SGraphNode::OnDrop(MyGeometry, DragDropEvent);
+}
 
-	InnerBlockSize.Set( (NodePadding.Left + NodePadding.Right + NodeDesiredWidth + InterNodeSpacing) * float(ActionNodes.Num() > 0 ? ActionNodes.Num() : 1), TrackDesiredHeight );
+FReply SDataprepGraphTrackNode::OnDragOver(const FGeometry & MyGeometry, const FDragDropEvent & DragDropEvent)
+{
+	TSharedPtr<FDataprepDragDropOp> DragActionStepNodeOp = DragDropEvent.GetOperationAs<FDataprepDragDropOp>();
+	if(DragActionStepNodeOp.IsValid())
+	{
+		DragActionStepNodeOp->SetHoveredNode(GraphNode);
+		TrackWidgetPtr->UpdateDragIndicator(DragDropEvent.GetScreenSpacePosition());
+	}
 
-	RightBlockSize.Set( InnerBlockSize.X, InnerBlockSize.Y);
+	return SGraphNode::OnDragOver(MyGeometry, DragDropEvent);
+}
 
-	const float ZoomedSizeInX = LocalSize.X * InvZoomAmount;
-	if(InnerBlockSize.X < ZoomedSizeInX)
+void SDataprepGraphTrackNode::OnDragLeave(const FDragDropEvent & DragDropEvent)
+{
+	TSharedPtr<FDataprepDragDropOp> DragActionStepNodeOp = DragDropEvent.GetOperationAs<FDataprepDragDropOp>();
+	if(DragActionStepNodeOp.IsValid())
 	{
-		RightBlockSize.X = ZoomedSizeInX - InnerBlockSize.X;
+		TrackWidgetPtr->ResetDragIndicator();
 	}
 
-	LeftBlockSize.Set( FMath::CeilToFloat( 10.f * InvZoomAmount ), InnerBlockSize.Y);
+	SGraphNode::OnDragLeave(DragDropEvent);
+}
 
-	float NodeMaxHeight = NodeDesiredWidth;
+void SDataprepGraphTrackNode::SetOwner(const TSharedRef<SGraphPanel>& OwnerPanel)
+{
+	ensure(!OwnerGraphPanelPtr.IsValid());
+
+	SGraphNode::SetOwner(OwnerPanel);
 
 	for(TSharedPtr<SDataprepGraphActionNode>& ActionNodePtr : ActionNodes)
 	{
-		if (SDataprepGraphActionNode* ActionNode = ActionNodePtr.Get())
+		if (ActionNodePtr.IsValid())
 		{
-			const float ActionNodeHeight = ActionNode->GetDesiredSize().Y;
-			if(NodeMaxHeight < ActionNodeHeight)
-			{
-				NodeMaxHeight = ActionNodeHeight;
-			}
+			ActionNodePtr->SetOwner(OwnerPanel);
 		}
 	}
+}
 
-	ReArrangeActionNodes();
+FVector2D SDataprepGraphTrackNode::Update(const FVector2D& LocalSize, float ZoomAmount)
+{
+	if(SDataprepGraphTrackWidget* TrackWidget = TrackWidgetPtr.Get())
+	{
+		RefreshLayout();
+
+		return TrackWidget->GetWorkingArea();
+	}
 
-	return FVector2D( InnerBlockSize.X, NodeMaxHeight);
+	return FVector2D(SDataprepGraphTrackWidget::NodeDesiredWidth);
 }
 
 FVector2D SDataprepGraphTrackNode::ComputeActionNodePosition(const FVector2D& InPosition)
 {
-	const float NewAbscissa = InPosition.X < NodeAbscissaMin ? NodeAbscissaMin : (InPosition.X > NodeAbscissaMax ? NodeAbscissaMax : InPosition.X);
+	ensure(TrackWidgetPtr);
 
-	return FVector2D( NewAbscissa, NodePadding.Top);
+	const float NewAbscissa = TrackWidgetPtr->ValidateNodeAbscissa(InPosition.X - NodePadding.Left - TrackWidgetOffset.X);
+	return FVector2D( NewAbscissa + NodePadding.Left + TrackWidgetOffset.X, GraphNode->NodePosY + TrackWidgetOffset.Y);
 }
 
 void SDataprepGraphTrackNode::OnStartNodeDrag(const TSharedRef<SDataprepGraphActionNode>& ActionNode)
@@ -357,218 +424,813 @@ void SDataprepGraphTrackNode::OnStartNodeDrag(const TSharedRef<SDataprepGraphAct
 	bNodeDragging = true;
 	bSkipNextDragUpdate = false;
  
-	const FGeometry& DesktopGeometry = GetOwnerPanel()->GetPersistentState().DesktopGeometry;
+	TrackWidgetPtr->OnStartNodeDrag(ActionNode->GetExecutionOrder());
+
 	LastDragScreenSpacePosition = FSlateApplication::Get().GetCursorPos();
+	const FGeometry& DesktopGeometry = GetOwnerPanel()->GetPersistentState().DesktopGeometry;
 	FVector2D DragLocalPosition = DesktopGeometry.AbsoluteToLocal( LastDragScreenSpacePosition );
 	DragOrdinate = DragLocalPosition.Y;
 
-	OriginalOrder = ActionNode->GetExecutionOrder();
-	CurrentOrder = OriginalOrder;
+	SGraphPanel* GraphPanel = OwnerGraphPanelPtr.Pin().Get();
+	ensure(GraphPanel);
 
-	NewActionsOrder.SetNum(ActionNodes.Num());
-	for(int32 Index = 0; Index < NewActionsOrder.Num(); ++Index)
+	GraphPanel->Invalidate(EInvalidateWidgetReason::Layout);
+}
+
+int32 SDataprepGraphTrackNode::OnEndNodeDrag()
+{
+	bNodeDragging = false;
+	bSkipNextDragUpdate = false;
+
+	SGraphPanel* GraphPanel = OwnerGraphPanelPtr.Pin().Get();
+	ensure(GraphPanel);
+
+	GraphPanel->Invalidate(EInvalidateWidgetReason::Layout);
+
+	return 	TrackWidgetPtr->OnEndNodeDrag();
+}
+
+void SDataprepGraphTrackNode::OnNodeDragged( TSharedPtr<SDataprepGraphActionNode>& ActionNodePtr, const FVector2D& DragScreenSpacePosition, const FVector2D& ScreenSpaceDelta)
+{
+	ensure(bNodeDragging);
+	ensure(ActionNodePtr.IsValid());
+
+	// This update is most likely due to a call to FSlateApplication::SetCursorPos. Skip it
+	if(bSkipNextDragUpdate || ScreenSpaceDelta.X == 0.f)
 	{
-		NewActionsOrder[Index] = Index;
+		bSkipNextDragUpdate = false;
+		return;
+	}
+
+	TSharedPtr<SGraphPanel> GraphPanel = GetOwnerPanel();
+	const FGeometry& DesktopGeometry = GraphPanel->GetCachedGeometry();
+	const FVector2D& Size = DesktopGeometry.GetLocalSize();
+	const float ZoomAmount = GraphPanel->GetZoomAmount();
+
+	float AbscissaDelta = TrackWidgetPtr->OnNodeDragged(ScreenSpaceDelta.X / ZoomAmount);
+		
+	FVector2D NewScreenSpaceDelta( AbscissaDelta * ZoomAmount, ScreenSpaceDelta.Y );
+
+	// Request the active panel to scroll if required
+	{
+		// Keep the mouse at the same position if dragged node has reached when of the ends of the track
+		if(AbscissaDelta == 0.f)
+		{
+			// Do nothing, the mouse's cursor and the dragged node should not move
+		}
+		// Panel is narrower than track, update panel if dragged node has entered non visible part of track
+		else if((TrackWidgetPtr->GetWorkingArea().X * ZoomAmount) > Size.X)
+		{
+			const FGeometry& NodeGeometry = ActionNodePtr->GetCachedGeometry();
+			FVector2D NodeAbsolutePosition = NodeGeometry.LocalToAbsolute(FVector2D(NodeGeometry.GetLocalSize().X, 0.f));
+			FVector2D NodeRelativePosition = DesktopGeometry.AbsoluteToLocal(NodeAbsolutePosition);
+
+			if(NodeRelativePosition.X > Size.X)
+			{
+				FVector2D ViewOffset = GraphPanel->GetViewOffset();
+				ViewOffset.X += NodeRelativePosition.X - DesktopGeometry.GetLocalSize().X;
+				GraphPanel->RestoreViewSettings(ViewOffset, ZoomAmount);
+			}
+			else
+			{
+				// Make sure cursor stays at a constant height
+				LastDragScreenSpacePosition.X = DragScreenSpacePosition.X;
+			}
+		}
+		else
+		{
+			// Make sure cursor stays at a constant height
+			LastDragScreenSpacePosition.X = DragScreenSpacePosition.X;
+		}
+
+		// Update cursor's position
+		bSkipNextDragUpdate = LastDragScreenSpacePosition.X != DragScreenSpacePosition.X;
+		FSlateApplication::Get().SetCursorPos( LastDragScreenSpacePosition );
 	}
 }
 
-void SDataprepGraphTrackNode::OnNodeDropped(bool bDropWasHandled)
+void SDataprepGraphTrackNode::OnControlKeyChanged(bool bControlKeyDown)
 {
-	FString SequenceText;
-	for(int32 Index = 0 ; Index < NewActionsOrder.Num(); ++Index)
+	if(bNodeDragging)
 	{
-		SequenceText += FString::Printf(TEXT("%d "), ActionNodes[NewActionsOrder[Index]]->GetExecutionOrder());
+		if(bControlKeyDown)
+		{
+			TrackWidgetPtr->OnControlKeyDown(true);
+		}
+		else
+		{
+			float MouseOffset = TrackWidgetPtr->OnControlKeyDown(false);
+
+			if(MouseOffset != 0.f)
+			{
+				// Update mouse position
+				LastDragScreenSpacePosition.X -= MouseOffset;
+
+				bSkipNextDragUpdate = true;
+				FSlateApplication::Get().SetCursorPos( LastDragScreenSpacePosition );
+			}
+		}
 	}
+}
+
+void SDataprepGraphTrackWidget::Construct(const FArguments& InArgs, TSharedPtr<SDataprepGraphTrackNode> InTrackNode)
+{
+	TrackNodePtr = InTrackNode;
+	ensure(TrackNodePtr.IsValid());
+
+	SHorizontalBox::Construct( SHorizontalBox::FArguments()
+		+SHorizontalBox::Slot()
+		.AutoWidth()
+		[
+			SAssignNew( TrackCanvas, SConstraintCanvas )
+		]
+	);
+
+	CanvasSlot = &TrackCanvas->AddSlot()
+	.Anchors( FAnchors( 0.f, 0.f, 0.f, 0.f ) )
+	.Offset( FMargin(-20.f, -50.f, 0.f, 0.f) )
+	.AutoSize(true)
+	.Alignment(FVector2D::ZeroVector)
+	.ZOrder(0)
+	[
+		SNew(SColorBlock)
+		.Color( FLinearColor::Transparent )
+		.Size( TAttribute<FVector2D>::Create(TAttribute<FVector2D>::FGetter::CreateSP( this, &SDataprepGraphTrackWidget::GetCanvasArea ) ) )
+	];
+
+	TrackCanvas->AddSlot()
+	.Anchors( FAnchors( 0.f, 0.f, 1.f, 0.f ) )
+	.Offset( FMargin(0.f, 5.f, 0.f, 0.f) )
+	.AutoSize(true)
+	.Alignment(FVector2D::ZeroVector)
+	.ZOrder(1)
+	[
+		SNew(SColorBlock)
+		.Color(FDataprepEditorStyle::GetColor( "Graph.TrackInner.BackgroundColor" ))
+		.Size( TAttribute<FVector2D>::Create(TAttribute<FVector2D>::FGetter::CreateSP( this, &SDataprepGraphTrackWidget::GetTrackArea ) ) )
+	];
+
+	LastDropSlot = INDEX_NONE;
+
+	DropIndicator = SNew(SBorder)
+	.BorderBackgroundColor( FDataprepEditorStyle::GetColor( "DataprepActionStep.DragAndDrop" ) )
+	.BorderImage(FEditorStyle::GetBrush("Graph.StateNode.Body"))
+	.HAlign(EHorizontalAlignment::HAlign_Center)
+	[
+		SNew(SBox)
+		.HeightOverride(TrackDesiredHeight - 2.f)
+		.WidthOverride(SlotSpacing - 2.f)
+	];
+
+	DropFiller = SNew(SColorBlock)
+	.Color(FLinearColor::Transparent)
+	.Size( FVector2D(SlotSpacing, TrackDesiredHeight) );
+
+	TSharedPtr<SColorBlock> NullNode = SNew(SColorBlock)
+	.Color( FLinearColor::Transparent )
+	.Size( FVector2D(NodeDesiredWidth, TrackDesiredHeight) );
 
-	if(bDropWasHandled && CurrentOrder != OriginalOrder)
+	CanvasArea = FVector2D(NodeDesiredWidth, NodeDesiredWidth);
+
+	float LeftOffset = 0.f;
+	if(UDataprepAsset* DataprepAsset = InTrackNode->GetDataprepAsset())
 	{
-		// #ueent_wip: Apply change from OriginalOrder to CurrentOrder onto Dataprep asset's array of actions
-		// And react to notification of change
-		for(int32 Index = 0; Index < NewActionsOrder.Num(); ++Index)
+		int32 ActionCount = DataprepAsset->GetActionCount();
+
+		TArray<TSharedPtr<SDataprepGraphActionNode>>& InActionNodes = InTrackNode->ActionNodes;
+
+		DropSlots.Reserve(ActionCount + 1);
+		ActionSlots.Reserve(ActionCount);
+		ActionNodes.Reserve(ActionCount);
+
+		for(int32 Index = 0; Index < ActionCount; ++Index)
 		{
-			if( SDataprepGraphActionNode* ActionWidget = ActionNodes[NewActionsOrder[Index]].Get())
+			// Add drop slot ahead of action node
 			{
-				ActionWidget->SetExecutionOrder(Index);
+				SConstraintCanvas::FSlot& Slot = TrackCanvas->AddSlot()
+				.Anchors( FAnchors( 0., 0.f, 0.f, 0.f ) )
+				.Offset( FMargin(LeftOffset, 5.f, SlotSpacing, TrackDesiredHeight ) )
+				.Alignment(FVector2D::ZeroVector)
+				.ZOrder(2)
+				[
+					DropFiller.ToSharedRef()
+				];
+				DropSlots.Add( &Slot );
 			}
+			LeftOffset += SlotSpacing;
+
+			// Add slot to host of action node
+			{
+				ActionNodes.Add(InActionNodes[Index].IsValid() ? InActionNodes[Index]->AsShared() : NullNode->AsShared());
+
+				SConstraintCanvas::FSlot& Slot = TrackCanvas->AddSlot()
+				.Anchors( FAnchors( 0., 0.f, 0.f, 0.f ) )
+				.Offset( FMargin(LeftOffset, 0.f, NodeDesiredWidth, NodeDesiredWidth ) )
+				.Alignment(FVector2D::ZeroVector)
+				.AutoSize(true)
+				.ZOrder(2)
+				[
+					ActionNodes[Index]
+				];
+				ActionSlots.Add( &Slot );
+			}
+			LeftOffset += NodeDesiredWidth;
 		}
+	}
 
-		ActionNodes.Sort([](const TSharedPtr<SDataprepGraphActionNode> A, const TSharedPtr<SDataprepGraphActionNode> B) { return A->GetExecutionOrder() < B->GetExecutionOrder(); });
+	CanvasArea.X += LeftOffset + TrackDesiredHeight;
+
+	// Add drop slot at end of track
+	{
+		SConstraintCanvas::FSlot& Slot = TrackCanvas->AddSlot()
+		.Anchors( FAnchors( 0., 0.f, 0.f, 0.f ) )
+		.Offset( FMargin(LeftOffset, 5.f, SlotSpacing, TrackDesiredHeight ) )
+		.Alignment(FVector2D::ZeroVector)
+		.ZOrder(2)
+		[
+			DropFiller.ToSharedRef()
+		];
+		DropSlots.Add( &Slot );
+		LeftOffset += SlotSpacing;
 	}
 
-	bNodeDragging = false;
-	NewActionsOrder.Empty(ActionNodes.Num());
+	// Add action slot at end of track
+	{
+		SConstraintCanvas::FSlot& Slot = TrackCanvas->AddSlot()
+		.Anchors( FAnchors( 0., 0.f, 0.f, 0.f ) )
+		.Offset( FMargin(LeftOffset, 0.f, NodeDesiredWidth, NodeDesiredWidth ) )
+		.Alignment(FVector2D::ZeroVector)
+		.AutoSize(true)
+		.ZOrder(2)
+		[
+			SNullWidget::NullWidget
+		];
+		ActionSlots.Add( &Slot );
+	}
+
+	// Add slot to host dragged node
+	DragSlot = &TrackCanvas->AddSlot()
+	.Anchors( FAnchors( 0., 0.f, 0.f, 0.f ) )
+	.Offset( FMargin(LeftOffset, 0.f, NodeDesiredWidth, NodeDesiredWidth ) )
+	.Alignment(FVector2D::ZeroVector)
+	.AutoSize(true)
+	.ZOrder(3)
+	[
+		SNullWidget::NullWidget
+	];
 
-	ReArrangeActionNodes();
+	DragStartIndex = INDEX_NONE;
+	DragIndex = INDEX_NONE;
 }
 
-void SDataprepGraphTrackNode::UpdatePanelOnDrag(const FVector2D& DragScreenSpacePosition, const FVector2D& ScreenSpaceDelta)
+void SDataprepGraphTrackWidget::MoveDropSlotTo(int32 SlotIndex)
 {
-	const float MaxAbscissa = LeftBlockSize.X + InnerBlockSize.X;
+	ensure(DropSlots.IsValidIndex(SlotIndex));
+	if(LastDropSlot >= 0)
+	{
+		(*DropSlots[LastDropSlot])[SNullWidget::NullWidget->AsShared()];
+	}
 
-	SDataprepGraphActionNode* ActionWidget = ActionNodes[NewActionsOrder[CurrentOrder]].Get();
-	FVector2D NodePosition = ActionWidget->GetPosition();
+	LastDropSlot = SlotIndex;
+	if(LastDropSlot >= 0)
+	{
+		(*DropSlots[LastDropSlot])[DropIndicator.ToSharedRef()];
+	}
 
-	TSharedPtr<SGraphPanel> GraphPanel = GetOwnerPanel();
+	Invalidate(EInvalidateWidgetReason::Layout);
+}
+
+void SDataprepGraphTrackWidget::SetActionSlot(int32 SlotIndex, const TSharedRef<SWidget>& Widget)
+{
+	ensure(ActionSlots.IsValidIndex(SlotIndex));
+	(*ActionSlots[SlotIndex])[Widget];
+
+	Invalidate(EInvalidateWidgetReason::Layout);
+}
+
+FVector2D SDataprepGraphTrackWidget::ComputeDesiredSize(float InLayoutScaleMultiplier) const
+{
+	//TrackCanvas->SlatePrepass(InLayoutScaleMultiplier);
+	return SHorizontalBox::ComputeDesiredSize(InLayoutScaleMultiplier);
+}
+
+FVector2D SDataprepGraphTrackWidget::GetActioNodeAbscissaRange()
+{
+	FVector2D Range(SlotSpacing);
+
+	if(SDataprepGraphTrackNode* TrackNode = TrackNodePtr.Pin().Get())
+	{
+		for(int32 Index = 0; Index < ActionNodes.Num() - 1; ++Index)
+		{
+			Range.Y += GetNodeSize(ActionNodes[Index]).X + SlotSpacing;
+		}
+	}
+
+	return Range;
+}
+
+int32 SDataprepGraphTrackWidget::GetHoveredActionNode(float LeftCorner, float Width)
+{
+	const float CenterAbscissa = LeftCorner + Width * 0.5f;
+
+	// Compute origin abscissa of last action node on track
+	for(int32 Index = 0; Index < ActionNodes.Num(); ++Index)
+	{
+		const FMargin SlotOffset = ActionSlots[Index]->OffsetAttr.Get(FMargin());
+		const float RelativeAbscissa = CenterAbscissa - SlotOffset.Left;
+		if(RelativeAbscissa > -SlotSpacing && RelativeAbscissa <= SlotOffset.Right)
+		{
+			return Index;
+		}
+	}
 
-	float ZoomAmount = GraphPanel->GetZoomAmount();
+	return bIsCopying ? ActionNodes.Num() : ActionNodes.Num() - 1;
+}
 
-	const FGeometry& DesktopGeometry = GraphPanel->GetPersistentState().DesktopGeometry;
-	const FVector2D& Size = DesktopGeometry.GetLocalSize() / ZoomAmount;
+float SDataprepGraphTrackWidget::ValidateNodeAbscissa(float InAbscissa)
+{
+	if(ActionSlots.Num() > 0)
+	{
+		AbscissaRange.X = ActionSlots[0]->OffsetAttr.Get(FVector2D::ZeroVector).Left;
+		AbscissaRange.Y = ActionSlots[ActionSlots.Num() - 1]->OffsetAttr.Get(FVector2D::ZeroVector).Left;
 
-	// Keep the mouse at the same position if dragged node has reached when of the ends of the track
-	if(NodePosition.X == NodeAbscissaMin || NodePosition.X == NodeAbscissaMax)
+		return InAbscissa < InAbscissa ? AbscissaRange.X : (InAbscissa > AbscissaRange.Y ? AbscissaRange.Y : InAbscissa);
+	}
+
+	return InAbscissa;
+}
+
+void SDataprepGraphTrackWidget::OnActionsOrderChanged(const TArray<TSharedPtr<SDataprepGraphActionNode>>& InActionNodes)
+{
+	ensure(InActionNodes.Num() == ActionNodes.Num());
+
+	for(int32 Index = 0; Index < ActionNodes.Num(); ++Index)
 	{
-		// Do nothing, the mouse's cursor and the dragged node should not move
+		ActionNodes[Index] = InActionNodes[Index]->AsShared();
 	}
-	// Panel is narrower than track, update panel if dragged node has entered non visible part of track
-	else if(MaxAbscissa > Size.X)
+
+	RefreshLayout(nullptr);
+}
+
+void SDataprepGraphTrackWidget::RefreshLayout(SGraphPanel* GraphPanel)
+{
+	WorkingArea = FVector2D(0.f, TrackDesiredHeight);
+	FMargin DropSlotOffset = DropSlots[0]->OffsetAttr.Get(FMargin());
+
+	float MaxWidth = 0.f;
+
+	for(int32 Index = 0; Index < ActionNodes.Num(); ++Index)
 	{
-		const float DragIncrementStep = NodeDesiredWidth / 6.f;
-		FVector2D ViewOffset = GraphPanel->GetViewOffset();
-		FVector2D DragPosition = DesktopGeometry.AbsoluteToLocal(DragScreenSpacePosition);
-		float AbscissaPanOffset = 0.f;
+		DropSlots[Index]->OffsetAttr.Set(DropSlotOffset);
+		DropSlotOffset.Left += SlotSpacing;
+
+		(*DropSlots[Index])[DropFiller.ToSharedRef()];
+
+		TSharedRef<SWidget>& ActionNode = ActionNodes[Index];
+
+		(*ActionSlots[Index])
+		[
+			ActionNode
+		];
 
-		const float LeftCornerAbscissa = NodePosition.X - InterNodeSpacing * 0.5f - ViewOffset.X;
-		const float RightCornerAbscissa = NodePosition.X + NodeDesiredWidth + InterNodeSpacing * 0.5f - ViewOffset.X;
+		const FVector2D Size = GetNodeSize(ActionNode);
 
-		// Dragged node's right corner is disappearing on the right, bring it back
-		if(  ScreenSpaceDelta.X > 0.f && RightCornerAbscissa > Size.X )
+		ActionSlots[Index]->OffsetAttr.Set( FMargin(DropSlotOffset.Left, 0.f, Size.X, Size.Y));
+		DropSlotOffset.Left += Size.X;
+
+		if(WorkingArea.Y < Size.Y)
 		{
-			// Compute offset to display right corner and bring right neighbor too if applicable
-			AbscissaPanOffset = CurrentOrder < (NewActionsOrder.Num() - 1) ? RightCornerAbscissa - Size.X + DragIncrementStep : RightCornerAbscissa - Size.X;
+			WorkingArea.Y = Size.Y;
 		}
-		// Dragged node's left corner is disappearing on the left, bring it back
-		else if( ScreenSpaceDelta.X < 0.f && LeftCornerAbscissa < 0.f )
+
+		if(MaxWidth < Size.X)
 		{
-			// Compute offset to display left corner and bring left neighbor's back if applicable
-			AbscissaPanOffset = CurrentOrder > 0 ? LeftCornerAbscissa - DragIncrementStep : LeftCornerAbscissa;
+			MaxWidth = Size.X;
 		}
-		else
+	}
+
+	DropSlots[ActionNodes.Num()]->OffsetAttr.Set(DropSlotOffset);
+	(*DropSlots[ActionNodes.Num()])[DropFiller.ToSharedRef()];
+
+	WorkingArea.X = DropSlotOffset.Left + SlotSpacing + MaxWidth/* * 0.5f + 1.f*/;
+	WorkingArea.Y += 30.f;
+
+	if(GraphPanel)
+	{
+		float ZoomAmount = GraphPanel->GetZoomAmount();
+		FVector2D PanelSize = GraphPanel->GetCachedGeometry().GetLocalSize() / ZoomAmount;
+
+		CanvasArea.X = PanelSize.X > WorkingArea.X ? PanelSize.X : WorkingArea.X;
+		CanvasArea.Y = PanelSize.Y > WorkingArea.Y ? PanelSize.Y : WorkingArea.Y;
+		CanvasArea = CanvasArea + FVector2D(40.f, 40.f / ZoomAmount);
+
+		CanvasSlot->OffsetAttr.Set(FMargin(-40.f, -40.f / ZoomAmount, 0.f, 0.f));
+	}
+
+	Invalidate(EInvalidateWidgetReason::Layout | EInvalidateWidgetReason::ChildOrder);
+}
+
+void SDataprepGraphTrackWidget::UpdateLayoutForMove()
+{
+	// Clear the leftest slot
+	ActionSlots[ActionNodes.Num()]->OffsetAttr.Set(FVector2D::ZeroVector);
+	(*ActionSlots[ActionNodes.Num()])[SNullWidget::NullWidget];
+
+	// Bring the dragged slot to the last fixed slot if applicable
+	FMargin DragSlotOffset = DragSlot->OffsetAttr.Get(FVector2D::ZeroVector);
+	if(DragSlotOffset.Left > AbscissaRange.Y)
+	{
+		DragSlotOffset.Left = AbscissaRange.Y;
+		DragSlot->OffsetAttr.Set(DragSlotOffset);
+	}
+
+	// Reconstruct action order if dragged node at left end
+	if(DragIndex == ActionNodes.Num())
+	{
+		for(int32 Index = 0; Index < DragStartIndex; ++Index)
 		{
-			// Make sure cursor stays at a constant height
-			LastDragScreenSpacePosition = DesktopGeometry.LocalToAbsolute( FVector2D( DragPosition.X, DragOrdinate ) );
+			NewActionsOrder[Index] = Index;
+		}
+
+		for(int32 Index = DragStartIndex; Index < DragIndex; ++Index)
+		{
+			NewActionsOrder[Index] = Index + 1;
 		}
 
-		// Apply offset to panel's canvas and move mouse to stay on top of dragged node if required
-		if(AbscissaPanOffset != 0.f)
+		--DragIndex;
+		NewActionsOrder[DragIndex] = DragStartIndex;
+	}
+
+	FMargin DropSlotOffset = DropSlots[0]->OffsetAttr.Get(FMargin());
+
+	for(int32 Index = 0; Index < ActionNodes.Num(); ++Index)
+	{
+		DropSlots[Index]->OffsetAttr.Set(DropSlotOffset);
+		DropSlotOffset.Left += SlotSpacing;
+
+		if(SConstraintCanvas::FSlot* ActionSlot = ActionSlots[Index])
 		{
-			// Compute new cursor's new screen space position
-			LastDragScreenSpacePosition = DesktopGeometry.LocalToAbsolute( FVector2D( DragPosition.X - AbscissaPanOffset, DragOrdinate ) );
+			FVector2D Size;
+
+			if(Index == DragIndex)
+			{
+				Size = GetNodeSize(DragSlot->GetWidget());
+
+				(*ActionSlot)
+					[
+						SNew(SColorBlock)
+						.Color( FLinearColor::Transparent )
+					.Size( Size )
+					];
+			}
+			else
+			{
+				TSharedRef<SWidget>& ActionNode = ActionNodes[NewActionsOrder[Index]];
+
+				Size = GetNodeSize(ActionNode);
+
+				(*ActionSlot)
+					[
+						ActionNode
+					];
+			}
 
-			// Pan panel accordingly
-			AbscissaPanOffset *= ZoomAmount;
-			FVector2D NewViewOffset(ViewOffset.X + AbscissaPanOffset, ViewOffset.Y);
-			GraphPanel->RestoreViewSettings(NewViewOffset, ZoomAmount);
+			ActionSlot->OffsetAttr.Set( FMargin(DropSlotOffset.Left, 0.f, Size.X, Size.Y));
+			DropSlotOffset.Left += Size.X;
+		}
+		else
+		{
+			ensure(false);
 		}
 	}
+
+	DropSlots[ActionNodes.Num()]->OffsetAttr.Set(DropSlotOffset);
+	DropSlotOffset.Left += SlotSpacing;
+
+	ActionSlots[ActionNodes.Num()]->OffsetAttr.Set(DropSlotOffset);
+	(*ActionSlots[ActionNodes.Num()])[SNullWidget::NullWidget];
+}
+
+void SDataprepGraphTrackWidget::UpdateLayoutForCopy()
+{
+	const FMargin DragSlotOffset = DragSlot->OffsetAttr.Get(FMargin());
+	const int32 HoveredIndex = GetHoveredActionNode( DragSlotOffset.Left, DragSlotOffset.Right);
+
+	FMargin SlotOffset = DropSlots[0]->OffsetAttr.Get(FMargin());
+
+	for(int32 Index = 0; Index < HoveredIndex; ++Index)
+	{
+		DropSlots[Index]->OffsetAttr.Set(SlotOffset);
+		SlotOffset.Left += SlotSpacing;
+
+		SConstraintCanvas::FSlot* ActionSlot = ActionSlots[Index];
+
+		const FVector2D Size = GetNodeSize(ActionNodes[Index]);
+
+		ActionSlot->OffsetAttr.Set(FMargin(SlotOffset.Left, 0.f, Size.X, Size.Y));
+		(*ActionSlot)[ActionNodes[Index]];
+
+		SlotOffset.Left += Size.X;
+	}
+
+	if(HoveredIndex == ActionNodes.Num())
+	{
+		DropSlots[ActionNodes.Num()]->OffsetAttr.Set(SlotOffset);
+		SlotOffset.Left += SlotSpacing;
+
+		const FVector2D Size = GetNodeSize(DragSlot->GetWidget());
+
+		SConstraintCanvas::FSlot* ActionSlot = ActionSlots[ActionNodes.Num()];
+		ActionSlot->OffsetAttr.Set(FMargin(SlotOffset.Left, 0.f, Size.X, Size.Y));
+		(*ActionSlot)
+		[
+			SNew(SColorBlock)
+			.Color( FLinearColor::Transparent )
+			.Size( Size )
+		];
+	}
 	else
 	{
-		// Make sure cursor stays at a constant height
-		LastDragScreenSpacePosition.X = DragScreenSpacePosition.X;
+		DropSlots[HoveredIndex]->OffsetAttr.Set(SlotOffset);
+		SlotOffset.Left += SlotSpacing;
+
+		FVector2D Size = GetNodeSize(DragSlot->GetWidget());
+
+		(*ActionSlots[HoveredIndex])
+		[
+			SNew(SColorBlock)
+			.Color( FLinearColor::Transparent )
+			.Size( Size )
+		];
+
+		ActionSlots[HoveredIndex]->OffsetAttr.Set(FMargin(SlotOffset.Left, 0.f, Size.X, Size.Y));
+		SlotOffset.Left += Size.X;
+
+		for(int32 Index = HoveredIndex + 1; Index <= ActionNodes.Num(); ++Index)
+		{
+			DropSlots[Index]->OffsetAttr.Set(SlotOffset);
+			SlotOffset.Left += SlotSpacing;
+
+			SConstraintCanvas::FSlot* ActionSlot = ActionSlots[Index];
+
+			Size = GetNodeSize(ActionNodes[Index - 1]);
+
+			ActionSlot->OffsetAttr.Set(FMargin(SlotOffset.Left, 0.f, Size.X, Size.Y));
+			(*ActionSlot)[ActionNodes[Index - 1]];
+
+			SlotOffset.Left += Size.X;
+		}
 	}
 
-	// Update cursor's position
-	bSkipNextDragUpdate = LastDragScreenSpacePosition.X != DragScreenSpacePosition.X;
-	FSlateApplication::Get().SetCursorPos( LastDragScreenSpacePosition );
+	Invalidate(EInvalidateWidgetReason::Layout | EInvalidateWidgetReason::Paint);
 }
 
-void SDataprepGraphTrackNode::OnNodeDragged( TSharedPtr<SDataprepGraphActionNode>& ActionNodePtr, const FVector2D& DragScreenSpacePosition, const FVector2D& ScreenSpaceDelta)
+void SDataprepGraphTrackWidget::OnStartNodeDrag(int32 IndexDragged)
 {
-	ensure(bNodeDragging);
-
-	// This update is most likely due to a call to FSlateApplication::SetCursorPos. Skip it
-	if(bSkipNextDragUpdate)
+	if(!DragSlot)
 	{
-		bSkipNextDragUpdate = false;
+		ensure(false);
 		return;
 	}
 
-	if(SDataprepGraphActionNode* ActionNode = ActionNodePtr.Get())
+	ensure(TrackNodePtr.IsValid());
+	ensure(ActionSlots.IsValidIndex(IndexDragged));
+
+	NewActionsOrder.SetNum(ActionNodes.Num() + 1);
+	for(int32 Index = 0; Index < NewActionsOrder.Num(); ++Index)
+	{
+		NewActionsOrder[Index] = Index;
+	}
+
+	DragStartIndex = IndexDragged;
+	DragIndex = IndexDragged;
+
+	FMargin Offset = ActionSlots[IndexDragged]->OffsetAttr.Get(FVector2D::ZeroVector);
+	DragSlot->OffsetAttr.Set(Offset);
+	(*DragSlot)[ActionNodes[IndexDragged]];
+
+	AbscissaRange.X = ActionSlots[0]->OffsetAttr.Get(FVector2D::ZeroVector).Left;
+
+	const FModifierKeysState ModifierKeyState = FSlateApplication::Get().GetModifierKeys();
+	bIsCopying = !(ModifierKeyState.IsControlDown() || ModifierKeyState.IsCommandDown());
+
+	OnControlKeyDown(!bIsCopying);
+}
+
+float SDataprepGraphTrackWidget::OnNodeDragged(float DragDelta)
+{
+	if(!DragSlot)
 	{
-		const FVector2D NodePosition = ActionNode->GetPosition();
+		ensure(false);
+		return 0.f;
+	}
+
+	FMargin DragSlotOffset = DragSlot->OffsetAttr.Get(FVector2D::ZeroVector);
+
+	float NewAbscissa = DragSlotOffset.Left + DragDelta;
+	 
+	bool bValidDrag = (NewAbscissa >= AbscissaRange.X) && (DragSlotOffset.Left < AbscissaRange.Y || NewAbscissa < AbscissaRange.Y);
+
+	if(bValidDrag)
+	{
+		TSharedPtr<SDataprepGraphActionNode> DraggedActionNode = StaticCastSharedRef<SDataprepGraphActionNode>(DragSlot->GetWidget());
+		ensure(DraggedActionNode.IsValid());
+
+		const float CachedNodeAbscissa = DragSlotOffset.Left;
+
+		DragSlotOffset.Left = NewAbscissa < AbscissaRange.X ? AbscissaRange.X : (NewAbscissa > AbscissaRange.Y ? AbscissaRange.Y : NewAbscissa);
 
-		bool bValidMove = 
-			// Dragged node on left end but mouse moving to the right
-			(NodePosition.X == NodeAbscissaMin && ScreenSpaceDelta.X > 0) ||
-			// Dragged node on right end but mouse moving to the left
-			(NodePosition.X == NodeAbscissaMax && ScreenSpaceDelta.X < 0) ||
-			// Dragged node within track
-			(NodePosition.X != NodeAbscissaMin && NodePosition.X != NodeAbscissaMax);
+		DragSlot->OffsetAttr.Set(DragSlotOffset);
+		DraggedActionNode->GetNodeObj()->NodePosX = DragSlotOffset.Left;
 
-		if(bValidMove)
+		// Check if center of dragged widget is over a neighboring widget by at least half its size
+		const int32 NewDragIndex = GetHoveredActionNode(DragSlotOffset.Left, DragSlotOffset.Right);
+
+		if(NewDragIndex != DragIndex)
 		{
-			const FVector2D NodeNewPosition = ComputeActionNodePosition( NodePosition + (ScreenSpaceDelta / GetOwnerPanel()->GetZoomAmount()));
+			ensure( NewActionsOrder.IsValidIndex(NewDragIndex) );
 
-			SNodePanel::SNode::FNodeSet NodeFilter;
-			ActionNode->SGraphNode::MoveTo( NodeNewPosition, NodeFilter);
+			// Record the swap
+			int32 TempOrder = NewActionsOrder[DragIndex];
+			NewActionsOrder[DragIndex] = NewActionsOrder[NewDragIndex];
+			NewActionsOrder[NewDragIndex] = TempOrder;
 
-			// Check if center of dragged widget is over a neighboring widget by at least half its size
-			const float NodeRelativeCenterAbscissa = NodeNewPosition.X + (NodeDesiredWidth * 0.5f) + (InterNodeSpacing * 0.5f) - (LeftBlockSize.X + NodePadding.Left);
-			const int32 NewOrder = FMath::FloorToInt( NodeRelativeCenterAbscissa / (NodeDesiredWidth + InterNodeSpacing));
+			const int32 CachedIndex = DragIndex;
+			DragIndex = NewDragIndex;
 
-			if(NewOrder != CurrentOrder)
+			if(bIsCopying)
+			{
+				UpdateLayoutForCopy();
+			}
+			else
 			{
-				ensure( NewActionsOrder.IsValidIndex(NewOrder) );
+				// Cache action node in new index
+				TSharedRef<SWidget> CachedActionNode = ActionSlots[DragIndex]->GetWidget();
 
-				// Make the swap
-				int32 TempOrder = NewActionsOrder[CurrentOrder];
-				NewActionsOrder[CurrentOrder] = NewActionsOrder[NewOrder];
-				NewActionsOrder[NewOrder] = TempOrder;
+				const FVector2D Size = GetNodeSize(DragSlot->GetWidget());
 
-				CurrentOrder = NewOrder;
+				(*ActionSlots[DragIndex])
+				[
+					SNew(SColorBlock)
+					.Color( FLinearColor::Transparent )
+					.Size( Size )
+				];
 
-				// Reflect swap in graph editor
-				ReArrangeActionNodes();
+				(*ActionSlots[CachedIndex])
+				[
+					CachedActionNode
+				];
 			}
 		}
 
-		// Request the active panel to scroll if required
-		UpdatePanelOnDrag( DragScreenSpacePosition, ScreenSpaceDelta );
+		TrackCanvas->Invalidate(EInvalidateWidgetReason::Layout | EInvalidateWidgetReason::ChildOrder);
+
+		return DragSlotOffset.Left - CachedNodeAbscissa;
 	}
+
+	return 0.f;
+}
+
+int32 SDataprepGraphTrackWidget::OnEndNodeDrag()
+{
+	if(!DragSlot)
+	{
+		ensure(false);
+		return INDEX_NONE;
+	}
+
+	TSharedPtr<SDataprepGraphActionNode> ActionNode = StaticCastSharedRef<SDataprepGraphActionNode>(DragSlot->GetWidget());
+	ensure(ActionNode.IsValid());
+
+	bIsCopying = false;
+	int32 CachedLastOrder = DragIndex;
+
+	const FVector2D Size = GetNodeSize(ActionNode->AsShared());
+
+	DragSlot->OffsetAttr.Set(FVector2D::ZeroVector);
+	(*DragSlot)[SNullWidget::NullWidget];
+
+	DropSlots[ActionNodes.Num()]->OffsetAttr.Set(FVector2D::ZeroVector);
+	(*DropSlots[ActionNodes.Num()])[SNullWidget::NullWidget];
+
+	RefreshLayout(nullptr);
+
+	DragStartIndex = INDEX_NONE;
+	DragIndex = INDEX_NONE;
+
+	NewActionsOrder.Reset();
+
+	TrackCanvas->Invalidate(EInvalidateWidgetReason::Layout | EInvalidateWidgetReason::ChildOrder);
+
+	return CachedLastOrder;
 }
 
-TSharedRef<FDragGraphActionNode> FDragGraphActionNode::New(const TSharedRef<SDataprepGraphTrackNode>& InTrackNodePtr, const TSharedRef<SDataprepGraphActionNode>& InDraggedNode)
+void SDataprepGraphTrackWidget::UpdateDragIndicator(FVector2D MousePosition)
 {
-	TSharedRef<FDragGraphActionNode> Operation = MakeShareable(new FDragGraphActionNode);
+	ResetDragIndicator();
 
-	Operation->TrackNodePtr = InTrackNodePtr;
-	Operation->ActionNodePtr = InDraggedNode;
+	float LocalAbscissa = GetCachedGeometry().AbsoluteToLocal(MousePosition).X - 1.f;
+
+	if(LocalAbscissa < DropSlots[0]->OffsetAttr.Get(FMargin()).Left)
+	{
+		(*DropSlots[0])[DropIndicator.ToSharedRef()];
+		return;
+	}
+
+	for(int32 Index = 0; Index < ActionNodes.Num(); ++Index)
+	{
+		FMargin SlotOffset = ActionSlots[Index]->OffsetAttr.Get(FMargin());
+		if(LocalAbscissa <= (SlotOffset.Left + SlotOffset.Right * 0.5f))
+		{
+			(*DropSlots[Index])[DropIndicator.ToSharedRef()];
+			return;
+		}
+	}
 
-	Operation->bCreateNewWindow = false;
-	Operation->Construct();
+	(*DropSlots[ActionNodes.Num()])[DropIndicator.ToSharedRef()];
+}
 
-	InTrackNodePtr->OnStartNodeDrag(InDraggedNode);
+int32 SDataprepGraphTrackWidget::GetDragIndicatorIndex()
+{
+	for(int32 Index = 0; Index < DropSlots.Num(); ++Index)
+	{
+		if(DropSlots[Index]->GetWidget() == DropIndicator)
+		{
+			return Index;
+		}
+	}
 
-	return Operation;
+	return -1;
 }
 
-TSharedRef<FDragGraphActionNode> FDragGraphActionNode::New(const TSharedRef<SDataprepGraphTrackNode>& InTrackNodePtr, const TArray< TSharedRef<SDataprepGraphActionNode> >& InDraggedNodes)
+float SDataprepGraphTrackWidget::OnControlKeyDown(bool bCtrlDown)
 {
-	TSharedRef<FDragGraphActionNode> Operation = MakeShareable(new FDragGraphActionNode);
+	if(bIsCopying == bCtrlDown)
+	{
+		return 0.f;
+	}
+
+	bIsCopying = bCtrlDown;
+
+	float MouseOffset = 0.f;
+
+	if(bIsCopying)
+	{
+		UpdateLayoutForCopy();
+
+		// Set the maximum move on the right
+		const FVector2D DragSize = GetNodeSize(DragSlot->GetWidget());
+		AbscissaRange.Y = ActionSlots[ActionNodes.Num()]->OffsetAttr.Get(FVector2D::ZeroVector).Left + 1.f - (DragSize.X * 0.5f);
+	}
+	else
+	{
+		AbscissaRange.Y = ActionSlots[ActionNodes.Num() - 1]->OffsetAttr.Get(FVector2D::ZeroVector).Left;
+
+		if(DragSlot->OffsetAttr.Get(FVector2D::ZeroVector).Left > AbscissaRange.Y)
+		{
+			// Update offset of DragSlot
+			FMargin DragSlotOffset = DragSlot->OffsetAttr.Get(FVector2D::ZeroVector);
+
+			MouseOffset = DragSlotOffset.Left - AbscissaRange.Y;
 
-	Operation->TrackNodePtr = InTrackNodePtr;
-	Operation->ActionNodePtr = InDraggedNodes[0];
+			DragSlotOffset.Left = AbscissaRange.Y;
+			DragSlot->OffsetAttr.Set(DragSlotOffset);
+		}
 
-	Operation->bCreateNewWindow = false;
-	Operation->Construct();
+		UpdateLayoutForMove();
+	}
 
-	InTrackNodePtr->OnStartNodeDrag(InDraggedNodes[0]);
+	TrackCanvas->Invalidate(EInvalidateWidgetReason::ChildOrder | EInvalidateWidgetReason::Layout);
 
-	return Operation;
+	return MouseOffset;
 }
 
-void FDragGraphActionNode::OnDrop(bool bDropWasHandled, const FPointerEvent& MouseEvent)
+bool SDataprepGraphTrackWidget::IsDraggedNodeOnEnds()
 {
-	TrackNodePtr->OnNodeDropped(bDropWasHandled);
+	if(DragSlot)
+	{
+		const FMargin DragSlotOffset = DragSlot->OffsetAttr.Get(FVector2D::ZeroVector);
+		return DragSlotOffset.Left == AbscissaRange.X || DragSlotOffset.Left == AbscissaRange.Y;
+	}
 
-	FDragDropOperation::OnDrop(bDropWasHandled, MouseEvent);
+	ensure(false);
+	return false;
 }
 
-void FDragGraphActionNode::OnDragged(const FDragDropEvent& DragDropEvent)
+FVector2D SDataprepGraphTrackWidget::GetNodeSize(const TSharedRef<SWidget>& Widget) const
 {
-	TrackNodePtr->OnNodeDragged( ActionNodePtr, DragDropEvent.GetScreenSpacePosition(), DragDropEvent.GetCursorDelta() );
+	FVector2D Size = Widget->GetCachedGeometry().GetLocalSize();
+
+	if(Size == FVector2D::ZeroVector)
+	{
+		Size = Widget->GetDesiredSize();
+		if(Size == FVector2D::ZeroVector)
+		{
+			Size.Set(NodeDesiredWidth, NodeDesiredWidth);
+		}
+	}
 
-	FDragDropOperation::OnDragged(DragDropEvent);
+	return Size;
 }
 
 #undef LOCTEXT_NAMESPACE
\ No newline at end of file
diff --git a/Engine/Plugins/Enterprise/DataprepEditor/Source/DataprepEditor/Private/Widgets/DataprepGraph/SDataprepGraphTrackNode.h b/Engine/Plugins/Enterprise/DataprepEditor/Source/DataprepEditor/Private/Widgets/DataprepGraph/SDataprepGraphTrackNode.h
--- a/Engine/Plugins/Enterprise/DataprepEditor/Source/DataprepEditor/Private/Widgets/DataprepGraph/SDataprepGraphTrackNode.h
+++ b/Engine/Plugins/Enterprise/DataprepEditor/Source/DataprepEditor/Private/Widgets/DataprepGraph/SDataprepGraphTrackNode.h
@@ -27,6 +27,30 @@ class UDataprepGraphActionK2Node;
 class UDataprepGraphActionNode;
 class UEdGraph;
 
+class FDragDropActionNode : public FDragDropOperation
+{
+public:
+	DRAG_DROP_OPERATOR_TYPE(FDragDropActionNode, FDragDropOperation)
+
+	static TSharedRef<FDragDropActionNode> New(const TSharedRef<SDataprepGraphTrackNode>& InTrackNodePtr, const TSharedRef<SDataprepGraphActionNode>& InDraggedNode);
+
+	virtual ~FDragDropActionNode() {}
+
+	// FDragDropOperation interface
+	virtual void OnDrop( bool bDropWasHandled, const FPointerEvent& MouseEvent ) override { Impl->OnDrop(bDropWasHandled, MouseEvent); }
+	virtual void OnDragged( const class FDragDropEvent& DragDropEvent ) override { Impl->OnDragged(DragDropEvent); }
+	virtual FCursorReply OnCursorQuery() override { return Impl->OnCursorQuery(); }
+	virtual TSharedPtr<SWidget> GetDefaultDecorator() const override { return Impl->GetDefaultDecorator(); }
+	virtual FVector2D GetDecoratorPosition() const override { return Impl->GetDecoratorPosition(); }
+	virtual void SetDecoratorVisibility(bool bVisible) override { Impl->SetDecoratorVisibility(bVisible); }
+	virtual bool IsExternalOperation() const override { return Impl->IsExternalOperation(); }
+	virtual bool IsWindowlessOperation() const override { return Impl->IsWindowlessOperation(); }
+	// End of FDragDropOperation interface
+
+private:
+	TSharedPtr<FDragDropActionNode> Impl;
+};
+
 /**
  * The SDataprepGraphTrackNode class is a specialization of SGraphNode
  * to handle the actions of a Dataprep asset
@@ -39,6 +63,10 @@ public:
 
 	void Construct(const FArguments& InArgs, UDataprepGraphRecipeNode* InNode);
 
+	// SGraphNode interface
+	virtual void UpdateGraphNode() override;
+	// End of SGraphNode interface
+
 	// SNodePanel::SNode interface
 	virtual bool CanBeSelected(const FVector2D& /*MousePositionInNode*/) const override
 	{
@@ -46,9 +74,17 @@ public:
 	}
 	virtual void SetOwner(const TSharedRef<SGraphPanel>& OwnerPanel) override;
 	virtual void MoveTo( const FVector2D& NewPosition, FNodeSet& NodeFilter ) override;
+	virtual const FSlateBrush* GetShadowBrush(bool bSelected) const override;
+	virtual FReply OnDrop( const FGeometry& MyGeometry, const FDragDropEvent& DragDropEvent ) override;
 	// End of SNodePanel::SNode interface
 
-	float GetInterNodeSpacing() const { return InterNodeSpacing; }
+	// SWidget interface
+	virtual FReply OnDragOver(const FGeometry& MyGeometry, const FDragDropEvent& DragDropEvent) override;
+	virtual void OnDragLeave(const FDragDropEvent& DragDropEvent) override;
+	// End of SWidget interface
+
+	UDataprepAsset* GetDataprepAsset() { return DataprepAssetPtr.Get(); }
+	const UDataprepAsset* GetDataprepAsset() const { return DataprepAssetPtr.Get(); }
 
 	/** Recompute the boundaries of the graph based on the new size and the new zoom factor */
 	FVector2D Update(const FVector2D& LocalSize, float ZoomAmount);
@@ -59,47 +95,34 @@ public:
 	 */
 	FVector2D ComputeActionNodePosition(const FVector2D& InPosition);
 
+	void OnControlKeyChanged(bool bControlKeyDown);
+
 	/** Initiates the horizontal drag of an action node */
 	void OnStartNodeDrag(const TSharedRef<SDataprepGraphActionNode>& ActionNode);
 
-	/** Reorganizes the array of action nodes based after a drag is completed */
-	void OnNodeDropped(bool bDropWasHandled);
+	/**
+	 * Terminates the horizontal drag of an action node
+	 * @return Returns the new execution order for the dragged action node
+	 */
+	int32 OnEndNodeDrag();
 
 	/** Updates the position of other action nodes based on the position of the incoming node */
 	void OnNodeDragged( TSharedPtr<SDataprepGraphActionNode>& ActionNodePtr, const FVector2D& DragScreenSpacePosition, const FVector2D& ScreenSpaceDelta);
 
+	/** Update the execution order of the actions and call ReArrangeActionNodes */
+	void OnActionsOrderChanged();
+
 	/** Recomputes the position of each action node */
-	void ReArrangeActionNodes();
+	bool RefreshLayout();
 
 	/** Miscellaneous values used in the display */
 	// #ueent_wip: Will be moved to the Dataprep editor's style
-	static float NodeDesiredWidth;
-	static float NodeDesiredSpacing;
 	static FMargin NodePadding;
-	static float TrackDesiredHeight;
-
-	// SGraphNode interface
-	virtual void UpdateGraphNode() override;
-	// End of SGraphNode interface
 
-private:
-	/** Updates the graph editor's canvas after a drag  */
-	void UpdatePanelOnDrag(const FVector2D& DragScreenSpacePosition, const FVector2D& ScreenSpaceDelta);
-
-	FVector2D GetInnerBlockSize() const
-	{
-		return InnerBlockSize;
-	}
-
-	FVector2D GetLeftBlockSize() const
-	{
-		return LeftBlockSize;
-	}
-
-	FVector2D GetRightBlockSize() const
-	{
-		return RightBlockSize;
-	}
+protected:
+	// SWidget interface
+	virtual bool CustomPrepass(float LayoutScaleMultiplier) override;
+	// End of SWidget interface
 
 private:
 	/** Pointer to the widget displaying the track */
@@ -111,23 +134,11 @@ private:
 	/** Weak pointer to the Dataprep asset holding the displayed actions */
 	TWeakObjectPtr<UDataprepAsset> DataprepAssetPtr;
 
-	/** Size of the section of the track encompassing the width of the action node's widgets */
-	FVector2D InnerBlockSize;
-
-	/** Size of the section of the track for right padding */
-	FVector2D RightBlockSize;
-
 	/** Size of the section of the track for left padding */
-	FVector2D LeftBlockSize;
+	FVector2D TrackWidgetOffset;
 
-	/** Spacing between action nodes. The value varies according to the zoom factor */
-	float InterNodeSpacing;
-
-	/** Minimum of position on the X axis */
-	float NodeAbscissaMin;
-
-	/** Maximum of position on the X axis */
-	float NodeAbscissaMax;
+	/** Range for abscissa of action nodes */
+	FVector2D AbscissaRange;
 
 	/** Indicates a drag is happening */
 	bool bNodeDragging;
@@ -155,21 +166,3 @@ private:
 
 	friend SDataprepGraphTrackWidget;
 };
-
-class FDragGraphActionNode : public FDragDropOperation
-{
-public:
-	DRAG_DROP_OPERATOR_TYPE(FDragGraphActionNode, FDragDropOperation)
-
-	static TSharedRef<FDragGraphActionNode> New(const TSharedRef<SDataprepGraphTrackNode>& InTrackNodePtr, const TSharedRef<SDataprepGraphActionNode>& InDraggedNode);
-	static TSharedRef<FDragGraphActionNode> New(const TSharedRef<SDataprepGraphTrackNode>& InTrackNodePtr, const TArray< TSharedRef<SDataprepGraphActionNode> >& InDraggedNodes);
-
-	// FDragDropOperation interface
-	virtual void OnDrop( bool bDropWasHandled, const FPointerEvent& MouseEvent );
-	virtual void OnDragged( const class FDragDropEvent& DragDropEvent );
-	// End of FDragDropOperation interface
-
-protected:
-	TSharedPtr<SDataprepGraphTrackNode> TrackNodePtr;
-	TSharedPtr<SDataprepGraphActionNode> ActionNodePtr;
-};
diff --git a/Engine/Plugins/Enterprise/DataprepEditor/Source/DataprepLibraries/Private/DataprepOperations.cpp b/Engine/Plugins/Enterprise/DataprepEditor/Source/DataprepLibraries/Private/DataprepOperations.cpp
--- a/Engine/Plugins/Enterprise/DataprepEditor/Source/DataprepLibraries/Private/DataprepOperations.cpp
+++ b/Engine/Plugins/Enterprise/DataprepEditor/Source/DataprepLibraries/Private/DataprepOperations.cpp
@@ -346,7 +346,7 @@ void UDataprepRandomizeTransformOperation::OnExecution_Implementation(const FDat
 #endif
 
 	// Execute operation
-	UDataprepOperationsLibrary::RandomizeTransform(InContext.Objects, TransformType, Mode, Min, Max);
+	UDataprepOperationsLibrary::RandomizeTransform(InContext.Objects, TransformType, ReferenceFrame, Min, Max);
 }
 
 #undef LOCTEXT_NAMESPACE
diff --git a/Engine/Plugins/Enterprise/DataprepEditor/Source/DataprepLibraries/Private/DataprepOperations.h b/Engine/Plugins/Enterprise/DataprepEditor/Source/DataprepLibraries/Private/DataprepOperations.h
--- a/Engine/Plugins/Enterprise/DataprepEditor/Source/DataprepLibraries/Private/DataprepOperations.h
+++ b/Engine/Plugins/Enterprise/DataprepEditor/Source/DataprepLibraries/Private/DataprepOperations.h
@@ -395,14 +395,14 @@ protected:
 	//~ End UDataprepOperation Interface
 };
 
-UCLASS(Experimental, Category = ActorOperation, Meta = (DisplayName = "Randomize Transform", ToolTip = "For each actor in the input set, offset its position/rotation/scale with random vector generated from X/Y/Z Min-Max."))
+UCLASS(Experimental, Category = ActorOperation, Meta = (DisplayName = "Random Offset Transform", ToolTip = "For each actor in the input set, offset its position/rotation/scale with random vector generated from X/Y/Z Min-Max."))
 class UDataprepRandomizeTransformOperation : public UDataprepOperation
 {
 	GENERATED_BODY()
 
 	UDataprepRandomizeTransformOperation()
-		: TransformType(ERandomizeTransformType::Translation) 
-		, Mode(ERandomizeTransformMode::Relative)
+		: TransformType(ERandomizeTransformType::Location) 
+		, ReferenceFrame(ERandomizeTransformReferenceFrame::Relative)
 	{
 	}
 
@@ -416,8 +416,8 @@ public:
 	UPROPERTY(EditAnywhere, Category = ActorOperation, Meta = (ToolTip = "Transform component to randomize"))
 	ERandomizeTransformType TransformType;
 
-	UPROPERTY(EditAnywhere, Category = ActorOperation, Meta = (ToolTip = "Transform mode to use (relative/world)"))
-	ERandomizeTransformMode Mode;
+	UPROPERTY(EditAnywhere, Category = ActorOperation, Meta = (ToolTip = "Reference frame to use (relative/world)"))
+	ERandomizeTransformReferenceFrame ReferenceFrame;
 
 	UPROPERTY(EditAnywhere, Category = ActorOperation, Meta = (ToolTip = "Min random value"))
 	FVector Min;
diff --git a/Engine/Plugins/Enterprise/DataprepEditor/Source/DataprepLibraries/Private/DataprepOperationsLibrary.cpp b/Engine/Plugins/Enterprise/DataprepEditor/Source/DataprepLibraries/Private/DataprepOperationsLibrary.cpp
--- a/Engine/Plugins/Enterprise/DataprepEditor/Source/DataprepLibraries/Private/DataprepOperationsLibrary.cpp
+++ b/Engine/Plugins/Enterprise/DataprepEditor/Source/DataprepLibraries/Private/DataprepOperationsLibrary.cpp
@@ -789,7 +789,7 @@ void UDataprepOperationsLibrary::ConsolidateObjects(const TArray< UObject* >& Se
 	ObjectTools::FConsolidationResults ConsResults = ObjectTools::ConsolidateObjects(ObjectToConsolidateTo, OutCompatibleObjects, false);
 }
 
-void UDataprepOperationsLibrary::RandomizeTransform(const TArray<UObject*>& SelectedObjects, ERandomizeTransformType TransformType, ERandomizeTransformMode Mode, const FVector& Min, const FVector& Max)
+void UDataprepOperationsLibrary::RandomizeTransform(const TArray<UObject*>& SelectedObjects, ERandomizeTransformType TransformType, ERandomizeTransformReferenceFrame ReferenceFrame, const FVector& Min, const FVector& Max)
 {
 	for (UObject* Object : SelectedObjects)
 	{
@@ -812,7 +812,7 @@ void UDataprepOperationsLibrary::RandomizeTransform(const TArray<UObject*>& Sele
 				case ERandomizeTransformType::Rotation:
 				{
 					const FRotator OffsetRotation = FRotator::MakeFromEuler(Offset);
-					if (Mode == ERandomizeTransformMode::World)
+					if (ReferenceFrame == ERandomizeTransformReferenceFrame::World)
 					{
 						RootComponent->SetWorldRotation(RootComponent->GetComponentRotation() + OffsetRotation);
 					}
@@ -824,7 +824,7 @@ void UDataprepOperationsLibrary::RandomizeTransform(const TArray<UObject*>& Sele
 				}
 				case ERandomizeTransformType::Scale:
 				{
-					if (Mode == ERandomizeTransformMode::World)
+					if (ReferenceFrame == ERandomizeTransformReferenceFrame::World)
 					{
 						RootComponent->SetWorldScale3D(RootComponent->GetComponentScale() + Offset);
 					}
@@ -834,9 +834,9 @@ void UDataprepOperationsLibrary::RandomizeTransform(const TArray<UObject*>& Sele
 					}
 					break;
 				}
-				case ERandomizeTransformType::Translation:
+				case ERandomizeTransformType::Location:
 				{
-					if (Mode == ERandomizeTransformMode::World)
+					if (ReferenceFrame == ERandomizeTransformReferenceFrame::World)
 					{
 						RootComponent->SetWorldLocation(RootComponent->GetComponentLocation() + Offset);
 					}
diff --git a/Engine/Plugins/Enterprise/DataprepEditor/Source/DataprepLibraries/Private/DataprepOperationsLibrary.h b/Engine/Plugins/Enterprise/DataprepEditor/Source/DataprepLibraries/Private/DataprepOperationsLibrary.h
--- a/Engine/Plugins/Enterprise/DataprepEditor/Source/DataprepLibraries/Private/DataprepOperationsLibrary.h
+++ b/Engine/Plugins/Enterprise/DataprepEditor/Source/DataprepLibraries/Private/DataprepOperationsLibrary.h
@@ -141,12 +141,12 @@ UENUM(BlueprintInternalUseOnly)
 enum class ERandomizeTransformType : uint8
 {
 	Rotation,
-	Translation,
+	Location,
 	Scale
 };
 
 UENUM(BlueprintInternalUseOnly)
-enum class ERandomizeTransformMode : uint8
+enum class ERandomizeTransformReferenceFrame : uint8
 {
 	World,
 	Relative
@@ -310,7 +310,7 @@ public:
 	 * @param Max:				End of randomization range
 	 */
 	UFUNCTION(BlueprintCallable, Category = "Dataprep | Operation")
-	static void RandomizeTransform(const TArray<UObject*>& SelectedObjects, ERandomizeTransformType TransformType, ERandomizeTransformMode Mode, const FVector& Min, const FVector& Max);
+	static void RandomizeTransform(const TArray<UObject*>& SelectedObjects, ERandomizeTransformType TransformType, ERandomizeTransformReferenceFrame ReferenceFrame, const FVector& Min, const FVector& Max);
 
 private:
 	static void SubstituteMaterial(const TArray<UObject*>& SelectedObjects, const FString& MaterialSearch, EEditorScriptingStringMatchType StringMatch, const TArray<UMaterialInterface*>& MaterialList, UMaterialInterface* MaterialSubstitute);
diff --git a/Engine/Plugins/Enterprise/DatasmithC4DImporter/Source/DatasmithC4DTranslator/Private/DatasmithC4DImportOptions.h b/Engine/Plugins/Enterprise/DatasmithC4DImporter/Source/DatasmithC4DTranslator/Private/DatasmithC4DImportOptions.h
--- a/Engine/Plugins/Enterprise/DatasmithC4DImporter/Source/DatasmithC4DTranslator/Private/DatasmithC4DImportOptions.h
+++ b/Engine/Plugins/Enterprise/DatasmithC4DImporter/Source/DatasmithC4DTranslator/Private/DatasmithC4DImportOptions.h
@@ -61,14 +61,3 @@ public:
 	bool bExportToUDatasmith;
 #endif //WITH_EDITORONLY_DATA
 };
-
-
-UCLASS(config = Editor, transient)
-class UDatasmithC4DReimportMeshOptions : public UObject
-{
-	GENERATED_BODY()
-
-public:
-	UPROPERTY(EditAnywhere, Category = StaticMesh, meta = (ShowOnlyInnerProperties))
-	FDatasmithStaticMeshImportOptions ImportOptions;
-};
diff --git a/Engine/Plugins/Enterprise/DatasmithC4DImporter/Source/DatasmithC4DTranslator/Private/DatasmithC4DTranslator.cpp b/Engine/Plugins/Enterprise/DatasmithC4DImporter/Source/DatasmithC4DTranslator/Private/DatasmithC4DTranslator.cpp
--- a/Engine/Plugins/Enterprise/DatasmithC4DImporter/Source/DatasmithC4DTranslator/Private/DatasmithC4DTranslator.cpp
+++ b/Engine/Plugins/Enterprise/DatasmithC4DImporter/Source/DatasmithC4DTranslator/Private/DatasmithC4DTranslator.cpp
@@ -28,7 +28,7 @@ bool FDatasmithC4DTranslator::LoadScene(TSharedRef<IDatasmithScene> OutScene)
 {
 	OutScene->SetHost(TEXT("C4DTranslator"));
 
-    Importer = MakeShared<FDatasmithC4DImporter>(OutScene, ImportOptions.Get());
+    Importer = MakeShared<FDatasmithC4DImporter>(OutScene, GetOrCreateC4DImportOptions().Get());
 	Importer->OpenFile(GetSource().GetSourceFile());
 
 	return Importer->ProcessScene();
@@ -71,12 +71,7 @@ bool FDatasmithC4DTranslator::LoadLevelSequence(const TSharedRef<IDatasmithLevel
 
 void FDatasmithC4DTranslator::GetSceneImportOptions(TArray<TStrongObjectPtr<UObject>>& Options)
 {
-	if (!ImportOptions.IsValid())
-	{
-		ImportOptions = Datasmith::MakeOptions<UDatasmithC4DImportOptions>();
-	}
-
-	Options.Add(ImportOptions);
+	Options.Add(GetOrCreateC4DImportOptions());
 }
 
 void FDatasmithC4DTranslator::SetSceneImportOptions(TArray<TStrongObjectPtr<UObject>>& Options)
@@ -92,8 +87,17 @@ void FDatasmithC4DTranslator::SetSceneImportOptions(TArray<TStrongObjectPtr<UObj
 
 	if (Importer.IsValid())
 	{
-		Importer->SetImportOptions(ImportOptions.Get());
+		Importer->SetImportOptions(GetOrCreateC4DImportOptions().Get());
+	}
+}
+
+TStrongObjectPtr<UDatasmithC4DImportOptions>& FDatasmithC4DTranslator::GetOrCreateC4DImportOptions()
+{
+	if (!ImportOptions.IsValid())
+	{
+		ImportOptions = Datasmith::MakeOptions<UDatasmithC4DImportOptions>();
 	}
+	return ImportOptions;
 }
 
 #endif // _MELANGE_SDK_
\ No newline at end of file
diff --git a/Engine/Plugins/Enterprise/DatasmithC4DImporter/Source/DatasmithC4DTranslator/Private/DatasmithC4DTranslator.h b/Engine/Plugins/Enterprise/DatasmithC4DImporter/Source/DatasmithC4DTranslator/Private/DatasmithC4DTranslator.h
--- a/Engine/Plugins/Enterprise/DatasmithC4DImporter/Source/DatasmithC4DTranslator/Private/DatasmithC4DTranslator.h
+++ b/Engine/Plugins/Enterprise/DatasmithC4DImporter/Source/DatasmithC4DTranslator/Private/DatasmithC4DTranslator.h
@@ -27,6 +27,9 @@ public:
 	virtual void GetSceneImportOptions(TArray<TStrongObjectPtr<UObject>>& Options) override;
 	virtual void SetSceneImportOptions(TArray<TStrongObjectPtr<UObject>>& Options) override;
 
+private:
+	TStrongObjectPtr<UDatasmithC4DImportOptions>& GetOrCreateC4DImportOptions();
+
 private:
 	TStrongObjectPtr<UDatasmithC4DImportOptions> ImportOptions;
 	TSharedPtr<class FDatasmithC4DImporter> Importer;
diff --git a/Engine/Plugins/Enterprise/DatasmithCADImporter/DatasmithCADImporter.uplugin b/Engine/Plugins/Enterprise/DatasmithCADImporter/DatasmithCADImporter.uplugin
--- a/Engine/Plugins/Enterprise/DatasmithCADImporter/DatasmithCADImporter.uplugin
+++ b/Engine/Plugins/Enterprise/DatasmithCADImporter/DatasmithCADImporter.uplugin
@@ -88,6 +88,14 @@
 			"WhitelistPlatforms": [
 				"Win64"
 			]
+		},
+    	{
+			"Name": "DatasmithPLMXMLTranslator",
+    		"Type": "Runtime",
+			"LoadingPhase": "Default",
+			"WhitelistPlatforms": [
+				"Win64"
+			]
 		}
 	],
 	"Plugins": [
diff --git a/Engine/Plugins/Enterprise/DatasmithCADImporter/Source/CADLibrary/Private/CoreTechHelper.cpp b/Engine/Plugins/Enterprise/DatasmithCADImporter/Source/CADLibrary/Private/CoreTechHelper.cpp
--- a/Engine/Plugins/Enterprise/DatasmithCADImporter/Source/CADLibrary/Private/CoreTechHelper.cpp
+++ b/Engine/Plugins/Enterprise/DatasmithCADImporter/Source/CADLibrary/Private/CoreTechHelper.cpp
@@ -76,6 +76,7 @@ void FillVertexPosition(const FImportParameters& ImportParams, const FMeshParame
 	// Add offset on the bbox to avoid to remove good vertex
 	FVector Size = Body.BBox.GetSize();
 	FBox BBox = Body.BBox.ExpandBy(Size.Size());
+	BBox.IsValid = Body.BBox.IsValid;
 	BBox.Min *= ImportParams.ScaleFactor;
 	BBox.Max *= ImportParams.ScaleFactor;
 
@@ -124,7 +125,7 @@ void FillVertexPosition(const FImportParameters& ImportParams, const FMeshParame
 
 		// Check if mesh vertex is not outside body bbox
 		const FVector& PositionA = VertexDataSet[i].Coordinates;
-		if (!BBox.IsInside(PositionA))
+		if (BBox.IsValid && !BBox.IsInside(PositionA))
 		{
 			VertexDataSet[i].Index = -1;
 			continue;
@@ -235,6 +236,23 @@ void FillVertexPosition(const FImportParameters& ImportParams, const FMeshParame
 	}
 }
 
+
+// PolygonAttributes name used into modeling tools (ExtendedMeshAttribute::PolyTriGroups)
+const FName PolyTriGroups("PolyTriGroups");
+
+// Copy of FMeshDescriptionBuilder::EnablePolyGroups()
+TPolygonAttributesRef<int32> EnableCADPatchGroups(FMeshDescription& MeshDescription)
+{
+	TPolygonAttributesRef<int32> PatchGroups = MeshDescription.PolygonAttributes().GetAttributesRef<int32>(PolyTriGroups);
+	if (PatchGroups.IsValid() == false)
+	{
+		MeshDescription.PolygonAttributes().RegisterAttribute<int32>(PolyTriGroups, 1, 0, EMeshAttributeFlags::AutoGenerated);
+		PatchGroups = MeshDescription.PolygonAttributes().GetAttributesRef<int32>(PolyTriGroups);
+		check(PatchGroups.IsValid());
+	}
+	return PatchGroups;
+}
+
 bool FillMesh(const FMeshParameters& MeshParameters, const FImportParameters& ImportParams, TArray<FTessellationData>& FaceTessellations, FMeshDescription& MeshDescription)
 {
 	const int32 UVChannel = 0;
@@ -264,6 +282,25 @@ bool FillMesh(const FMeshParameters& MeshParameters, const FImportParameters& Im
 		return false;
 	}
 
+	// Find all the materials used
+	TMap<uint32, FPolygonGroupID> MaterialToPolygonGroupMapping;
+	for (const FTessellationData& FaceTessellation : FaceTessellations)
+	{
+		// we assume that face has only color
+		MaterialToPolygonGroupMapping.Add(FaceTessellation.ColorName, FPolygonGroupID::Invalid);
+	}
+
+	// Add to the mesh, a polygon groups per material
+	for (auto& Material : MaterialToPolygonGroupMapping)
+	{
+		uint32 MaterialHash = Material.Key;
+		FName ImportedSlotName = *LexToString<uint32>(MaterialHash);
+
+		FPolygonGroupID PolyGroupID = MeshDescription.CreatePolygonGroup();
+		PolygonGroupImportedMaterialSlotNames[PolyGroupID] = ImportedSlotName;
+		Material.Value = PolyGroupID;
+	}
+
 	VertexInstanceUVs.SetNumIndices(1);
 
 	int32 NbStep = 1;
@@ -272,6 +309,8 @@ bool FillMesh(const FMeshParameters& MeshParameters, const FImportParameters& Im
 		NbStep = 2;
 	}
 
+	TPolygonAttributesRef<int32> PatchGroups = EnableCADPatchGroups(MeshDescription);
+	int32 PatchIndex = 0;
 	for (int32 Step = 0; Step < NbStep; ++Step)
 	{
 		// Swap mesh if needed
@@ -279,10 +318,13 @@ bool FillMesh(const FMeshParameters& MeshParameters, const FImportParameters& Im
 
 		for (FTessellationData& Tessellation : FaceTessellations)
 		{
-			FPolygonGroupID PolygonGroupID = MeshDescription.CreatePolygonGroup();
-			CADUUID MaterialUUID = Tessellation.ColorName;
-			FName ImportedSlotName = *LexToString<uint32>(MaterialUUID);
-			PolygonGroupImportedMaterialSlotNames[PolygonGroupID] = ImportedSlotName;
+			PatchIndex++;
+			// Get the polygonGroup
+			const FPolygonGroupID* PolygonGroupID = MaterialToPolygonGroupMapping.Find(Tessellation.ColorName);
+			if (PolygonGroupID == nullptr)
+			{
+				continue;
+			}
 
 			//int32 TriangleCount = IndicesCount / 3;
 			int32 VertexIDs[3];
@@ -324,7 +366,9 @@ bool FillMesh(const FMeshParameters& MeshParameters, const FImportParameters& Im
 				TriangleVertexInstanceIDs[2] = MeshVertexInstanceIDs[NewIndex++] = MeshDescription.CreateVertexInstance((FVertexID) VertexIDs[2]);
 
 				// Add the triangle as a polygon to the mesh description
-				MeshDescription.CreateTriangle(PolygonGroupID, TriangleVertexInstanceIDs);
+				const FPolygonID PolygonID = MeshDescription.CreatePolygon(*PolygonGroupID, TriangleVertexInstanceIDs);
+				// Set patch id attribute
+				PatchGroups[PolygonID] = PatchIndex;
 			}
 
 			// finalization of the mesh by setting colors, tangents, bi-normals, UV
diff --git a/Engine/Plugins/Enterprise/DatasmithCADImporter/Source/DatasmithOpenNurbsTranslator/DatasmithOpenNurbsTranslator.Build.cs b/Engine/Plugins/Enterprise/DatasmithCADImporter/Source/DatasmithOpenNurbsTranslator/DatasmithOpenNurbsTranslator.Build.cs
--- a/Engine/Plugins/Enterprise/DatasmithCADImporter/Source/DatasmithOpenNurbsTranslator/DatasmithOpenNurbsTranslator.Build.cs
+++ b/Engine/Plugins/Enterprise/DatasmithCADImporter/Source/DatasmithOpenNurbsTranslator/DatasmithOpenNurbsTranslator.Build.cs
@@ -28,7 +28,7 @@ namespace UnrealBuildTool.Rules
 					"DatasmithContent",
 					"DatasmithTranslator",
 					"MeshDescription",
-					"StaticMeshDescription"
+					"StaticMeshDescription",
 				}
 			);
 
@@ -41,6 +41,12 @@ namespace UnrealBuildTool.Rules
             {
 				PrivateDependencyModuleNames.Add("CoreTech");
             }
+
+            if (Target.Type == TargetType.Editor)
+            {
+	            PrivateDependencyModuleNames.Add("MessageLog");
+            }
+
         }
     }
 }
diff --git a/Engine/Plugins/Enterprise/DatasmithCADImporter/Source/DatasmithOpenNurbsTranslator/Private/DatasmithOpenNurbsImportOptions.cpp b/Engine/Plugins/Enterprise/DatasmithCADImporter/Source/DatasmithOpenNurbsTranslator/Private/DatasmithOpenNurbsImportOptions.cpp
--- a/Engine/Plugins/Enterprise/DatasmithCADImporter/Source/DatasmithOpenNurbsTranslator/Private/DatasmithOpenNurbsImportOptions.cpp
+++ b/Engine/Plugins/Enterprise/DatasmithCADImporter/Source/DatasmithOpenNurbsTranslator/Private/DatasmithOpenNurbsImportOptions.cpp
@@ -18,9 +18,9 @@ bool UDatasmithOpenNurbsImportOptions::CanEditChange(const FProperty* InProperty
 
 	FName PropertyFName = InProperty->GetFName();
 
-	if (PropertyFName == GET_MEMBER_NAME_CHECKED(FDatasmithOpenNurbsOptions, BrepTesselation))
+	if (PropertyFName == GET_MEMBER_NAME_CHECKED(FDatasmithOpenNurbsOptions, Geometry))
 	{
-#ifdef CAD_LIBRARY		
+#ifdef CAD_LIBRARY
 		return true;
 #else
 		return false;
@@ -33,7 +33,7 @@ bool UDatasmithOpenNurbsImportOptions::CanEditChange(const FProperty* InProperty
 		)
 	{
 		// Enable tessellation options only when using CAD library to tessellate
-		return Options.BrepTesselation == EDatasmithOpenNurbsBrepTesselatedSource::UseUnrealNurbsTessellation;
+		return Options.Geometry == EDatasmithOpenNurbsBrepTessellatedSource::UseUnrealNurbsTessellation;
 	}
 
 	return true;
diff --git a/Engine/Plugins/Enterprise/DatasmithCADImporter/Source/DatasmithOpenNurbsTranslator/Private/DatasmithOpenNurbsImportOptions.h b/Engine/Plugins/Enterprise/DatasmithCADImporter/Source/DatasmithOpenNurbsTranslator/Private/DatasmithOpenNurbsImportOptions.h
--- a/Engine/Plugins/Enterprise/DatasmithCADImporter/Source/DatasmithOpenNurbsTranslator/Private/DatasmithOpenNurbsImportOptions.h
+++ b/Engine/Plugins/Enterprise/DatasmithCADImporter/Source/DatasmithOpenNurbsTranslator/Private/DatasmithOpenNurbsImportOptions.h
@@ -12,13 +12,13 @@
 
 
 UENUM()
-enum class EDatasmithOpenNurbsBrepTesselatedSource : uint8
+enum class EDatasmithOpenNurbsBrepTessellatedSource : uint8
 {
 	/** Tessellate all Breps on import */
-	UseUnrealNurbsTessellation,
+	UseUnrealNurbsTessellation UMETA(DisplayName = "Import as NURBS, Tessellate in Unreal"),
 
-	/** Use meshes stored in the scene file */
-	UseSceneRenderMeshes,
+	/** Use render meshes stored in the scene file */
+	UseRenderMeshes UMETA(DisplayName = "Import Rhino Meshes and UVs"),
 };
 
 
@@ -30,18 +30,18 @@ struct FDatasmithOpenNurbsOptions : public FDatasmithTessellationOptions
 	FDatasmithOpenNurbsOptions()
     {
 #ifdef CAD_LIBRARY
-		BrepTesselation = EDatasmithOpenNurbsBrepTesselatedSource::UseUnrealNurbsTessellation;
+		Geometry = EDatasmithOpenNurbsBrepTessellatedSource::UseUnrealNurbsTessellation;
 #else
-		BrepTesselation = EDatasmithOpenNurbsBrepTesselatedSource::UseSceneRenderMeshes;
+		Geometry = EDatasmithOpenNurbsBrepTessellatedSource::UseRenderMeshes;
 #endif
     }
 
-	UPROPERTY(config, EditAnywhere, SimpleDisplay, Category = "Geometry & Tessellation Options")
-	EDatasmithOpenNurbsBrepTesselatedSource BrepTesselation;
+	UPROPERTY(config, EditAnywhere, BlueprintReadWrite, Category = "Geometry & Tessellation Options")
+	EDatasmithOpenNurbsBrepTessellatedSource Geometry;
 
 	uint32 GetHash() const
 	{
-		return HashCombine(FDatasmithTessellationOptions::GetHash(), GetTypeHash(BrepTesselation));;
+		return HashCombine(FDatasmithTessellationOptions::GetHash(), GetTypeHash(Geometry));;
 	}
 };
 
diff --git a/Engine/Plugins/Enterprise/DatasmithCADImporter/Source/DatasmithOpenNurbsTranslator/Private/DatasmithOpenNurbsTranslator.cpp b/Engine/Plugins/Enterprise/DatasmithCADImporter/Source/DatasmithOpenNurbsTranslator/Private/DatasmithOpenNurbsTranslator.cpp
--- a/Engine/Plugins/Enterprise/DatasmithCADImporter/Source/DatasmithOpenNurbsTranslator/Private/DatasmithOpenNurbsTranslator.cpp
+++ b/Engine/Plugins/Enterprise/DatasmithCADImporter/Source/DatasmithOpenNurbsTranslator/Private/DatasmithOpenNurbsTranslator.cpp
@@ -3,9 +3,13 @@
 #include "DatasmithOpenNurbsTranslator.h"
 #include "DatasmithOpenNurbsTranslatorModule.h"
 
+#ifdef USE_OPENNURBS // The whole translation unit is skipped without OpenNurbs TPS library
+
 #ifdef CAD_LIBRARY
 #include "CoreTechParametricSurfaceExtension.h"
-#endif
+#include "RhinoCoretechWrapper.h"
+#endif // CAD_LIBRARY
+
 #include "DatasmithImportOptions.h"
 #include "DatasmithMaterialElements.h"
 #include "DatasmithMaterialsUtils.h"
@@ -15,22 +19,20 @@
 #include "DatasmithUtils.h"
 #include "Utility/DatasmithMeshHelper.h"
 
-#include "StaticMeshAttributes.h"
-#include "StaticMeshOperations.h"
+#if WITH_EDITOR
+#include "IMessageLogListing.h"
+#include "MessageLogModule.h"
+#endif
 #include "Misc/FileHelper.h"
 #include "Misc/Paths.h"
+#include "StaticMeshAttributes.h"
+#include "StaticMeshOperations.h"
 
-#ifdef USE_OPENNURBS
 // Disable macro redefinition warning
 #pragma warning(push)
 #pragma warning(disable:4005)
 #include "opennurbs.h"
 #pragma warning(pop)
-#endif
-
-#ifdef CAD_LIBRARY
-#include "RhinoCoretechWrapper.h" // requires CoreTech as public dependency
-#endif
 
 #include <deque>
 #include <map>
@@ -38,6 +40,8 @@
 
 DEFINE_LOG_CATEGORY_STATIC(LogDatasmithOpenNurbsTranslator, Log, All);
 
+#define LOCTEXT_NAMESPACE "DatasmithOpenNurbsTranslator"
+
 // Cache for already processed data (only linked file references for now)
 class FTranslationCache
 {
@@ -82,7 +86,7 @@ TSharedPtr<IDatasmithActorElement> DuplicateActorElement(TSharedPtr<IDatasmithAc
 	DuplicatedElement->SetTranslation(ActorElement->GetTranslation());
 	DuplicatedElement->SetScale(ActorElement->GetScale());
 	DuplicatedElement->SetRotation(ActorElement->GetRotation());
-	
+
 	int32 NumChildren = ActorElement->GetChildrenCount();
 	for (int32 Index = 0; Index < NumChildren; ++Index)
 	{
@@ -99,7 +103,6 @@ TSharedPtr<IDatasmithActorElement> DuplicateActorElement(TSharedPtr<IDatasmithAc
 	return DuplicatedElement;
 }
 
-#ifdef USE_OPENNURBS
 // #ueent_wip: test with CADSDK_ENABLED undefined
 class FOpenNurbsObjectWrapper
 {
@@ -620,14 +623,14 @@ public:
 
 #ifdef CAD_LIBRARY
 		LocalSession = FRhinoCoretechWrapper::GetSharedSession(MetricUnit, ScalingFactor);
-#endif
+#endif // CAD_LIBRARY
 	}
 
 	~FOpenNurbsTranslatorImpl()
 	{
 #ifdef CAD_LIBRARY
 		LocalSession.Reset();
-#endif
+#endif // CAD_LIBRARY
 
 		for (FOpenNurbsTranslatorImpl* ChildTranslator : ChildTranslators)
 		{
@@ -645,6 +648,8 @@ public:
 	double GetScalingFactor () const { return ScalingFactor; }
 	double GetMetricUnit() const { return MetricUnit; }
 
+	void ShowMessageLog(const FString& Filename);
+
 private:
 	void TranslateTextureMappingTable(const ON_ObjectArray<ON_TextureMapping>& TextureMappingTable);
 	void TranslateMaterialTable(const ON_ObjectArray<ON_Material>& MaterialTable);
@@ -674,7 +679,7 @@ private:
 	FString GetLayerName(const TSharedPtr<IDatasmithActorElement>& LayerElement);
 	void SetLayers(const TSharedPtr<IDatasmithActorElement>& ActorElement, const FOpenNurbsObjectWrapper& Object);
 	void SetTags(const TSharedPtr<IDatasmithActorElement>& ActorElement, const FOpenNurbsObjectWrapper& Object);
-	
+
 	bool TranslateBRep(ON_Brep* brep, const ON_3dmObjectAttributes& Attributes, FMeshDescription& OutMesh, const TSharedRef< IDatasmithMeshElement >& MeshElement, const FString& Name, bool& bHasNormal);
 
 private:
@@ -690,7 +695,7 @@ private:
 
 #ifdef CAD_LIBRARY
 	TSharedPtr<FRhinoCoretechWrapper> LocalSession;
-#endif
+#endif // CAD_LIBRARY
 
 private:
 	// For OpenNurbs archive parsing
@@ -763,8 +768,36 @@ private:
 
 	/** OpenNurbs objects to Datasmith mesh elements */
 	TMap<const FOpenNurbsObjectWrapper* , TSharedPtr< IDatasmithMeshElement > > ObjectToMeshElementMap;
+
+	TArray<FString> MissingRenderMeshes;
 };
 
+void FOpenNurbsTranslatorImpl::ShowMessageLog(const FString& Filename)
+{
+#if WITH_EDITOR
+	if (MissingRenderMeshes.Num())
+	{
+		FMessageLogModule& MessageLogModule = FModuleManager::LoadModuleChecked<FMessageLogModule>("MessageLog");
+		TSharedRef<IMessageLogListing> LogListing = (MessageLogModule.GetLogListing(FName(TEXT("DatasmithOpenNurbs"))));
+		LogListing->SetLabel(LOCTEXT("DatasmithOpenNurbsTranslatorDescription", "DatasmithOpenNurbs"));
+
+		LogListing->ClearMessages();
+
+		LogListing->AddMessage(FTokenizedMessage::Create(EMessageSeverity::Warning,
+			FText::Format(LOCTEXT("DatasmithOpenNurbsTranslator_NoMeshDataForAllMeshes",
+				"Rhino model \"{0}\" doesn't contain mesh data for all objects. \n"
+				"Either resave the 3dm file with a \"rendered view\" or change the import settings to \"Import as NURBS, Tessellate in Unreal\""),
+				FText::FromString(Filename))));
+
+		for (const FString& Name: MissingRenderMeshes)
+		{
+			FText ErrorMessage = FText::Format(LOCTEXT("DatasmithOpenNurbsTranslator_NoMesh", "  {0} doesn't have mesh information."), FText::FromString(Name));
+			LogListing->AddMessage(FTokenizedMessage::Create(EMessageSeverity::Error, ErrorMessage));
+		}
+	}
+#endif
+}
+
 void FOpenNurbsTranslatorImpl::TranslateTextureMappingTable(const ON_ObjectArray<ON_TextureMapping>& InTextureMappingTable)
 {
 	for (int Index = 0; Index < InTextureMappingTable.Count(); ++Index)
@@ -1534,53 +1567,55 @@ void FOpenNurbsTranslatorImpl::SetTags(const TSharedPtr<IDatasmithActorElement>&
 	FString UUID(UUIDString.Array());
 
 	ON::object_type objType = Object.ObjectPtr->ObjectType();
-	FString StrObjectType;
+	const TCHAR* StrObjectType;
 	switch (objType)
 	{
-		case ON::instance_definition: 
-			StrObjectType = "block definition"; 
+		case ON::instance_definition:
+			StrObjectType = TEXT("block definition");
+			break;
+		case ON::instance_reference:
+			StrObjectType = TEXT("block instance");
 			break;
-		case ON::instance_reference: 
-			StrObjectType = "block instance"; 
+		case ON::point_object:
+			StrObjectType = TEXT("point");
 			break;
-		case ON::point_object: 
-			StrObjectType = "point"; 
+		case ON::curve_object:
+			StrObjectType = TEXT("curve");
 			break;
-		case ON::curve_object: 
-			StrObjectType = "curve"; 
+		case ON::surface_object:
+			StrObjectType = TEXT("surface");
 			break;
-		case ON::surface_object: 
-			StrObjectType = "surface"; 
+		case ON::brep_object:
+			StrObjectType = TEXT("brep");
 			break;
-		case ON::brep_object: 
-			StrObjectType = "brep"; 
+		case ON::mesh_object:
+			StrObjectType = TEXT("mesh");
 			break;
-		case ON::mesh_object: 
-			StrObjectType = "mesh"; 
+		case ON::text_dot:
+			StrObjectType = TEXT("textdot");
 			break;
-		case ON::text_dot: 
-			StrObjectType = "textdot"; 
+		case ON::subd_object:
+			StrObjectType = TEXT("subd");
 			break;
-		case ON::subd_object: 
-			StrObjectType = "subd"; 
+		case ON::loop_object:
+			StrObjectType = TEXT("loop");
 			break;
-		case ON::loop_object: 
-			StrObjectType = "loop"; 
+		case ON::cage_object:
+			StrObjectType = TEXT("cage");
 			break;
-		case ON::cage_object: 
-			StrObjectType = "cage"; 
+		case ON::clipplane_object:
+			StrObjectType = TEXT("clip plane");
 			break;
-		case ON::clipplane_object: 
-			StrObjectType = "clip plane"; 
+		case ON::extrusion_object:
+			StrObjectType = TEXT("extrusion");
 			break;
-		case ON::extrusion_object: 
-			StrObjectType = "extrusion"; 
+		default:
+			StrObjectType = TEXT("unknown");
 			break;
-		default: "unknown"; break;
 	}
 
 	ActorElement->AddTag(*FString::Printf(TEXT("Rhino.ID: %s"), *UUID));
-	ActorElement->AddTag(*FString::Printf(TEXT("Rhino.Entity.Type: %s"), *StrObjectType));
+	ActorElement->AddTag(*FString::Printf(TEXT("Rhino.Entity.Type: %s"), StrObjectType));
 }
 
 void FOpenNurbsTranslatorImpl::TranslateNonInstanceObject(const FOpenNurbsObjectWrapper& Object)
@@ -1874,7 +1909,7 @@ TSharedPtr<IDatasmithActorElement> FOpenNurbsTranslatorImpl::GetPointActorElemen
 	if (Object.Attributes.m_name.Length() > 0)
 	{
 		ActorLabel = Object.Attributes.m_name.Array();
-	} 
+	}
 	else
 	{
 		const ON_ClassId* classId = Object.ObjectPtr->ClassId();
@@ -1895,10 +1930,10 @@ TSharedPtr<IDatasmithActorElement> FOpenNurbsTranslatorImpl::GetPointActorElemen
 	FVector Location(pointObj->point.x, pointObj->point.y, pointObj->point.z);
 	Location *= ScalingFactor;
 	Location = FDatasmithUtils::ConvertVector(FDatasmithUtils::EModelCoordSystem::ZUp_RightHanded, Location);
-	
+
 	ActorElement->SetTranslation(Location);
 	ActorElement->SetLabel(*ActorLabel);
-	
+
 	return ActorElement;
 }
 
@@ -2111,7 +2146,7 @@ bool FOpenNurbsTranslatorImpl::Read(ON_BinaryFile& Archive, TSharedRef<IDatasmit
 	ScalingFactor = 100 / Settings.m_ModelUnitsAndTolerances.Scale(ON::LengthUnitSystem::Meters);
 #ifdef CAD_LIBRARY
 	LocalSession->SetScaleFactor(ScalingFactor);
-#endif
+#endif // CAD_LIBRARY
 
 	// Step 4: REQUIRED - Read bitmap table (it can be empty)
 	int Count = 0;
@@ -2959,8 +2994,7 @@ bool FOpenNurbsTranslatorImpl::TranslateBRep(ON_Brep* Brep, const ON_3dmObjectAt
 
 	// No tessellation if CAD library is not present...
 #ifdef CAD_LIBRARY
-	EDatasmithOpenNurbsBrepTesselatedSource bUseCADLibraryTesselation = OpenNurbsOptions.BrepTesselation;
-	if (bUseCADLibraryTesselation == EDatasmithOpenNurbsBrepTesselatedSource::UseUnrealNurbsTessellation)
+	if (OpenNurbsOptions.Geometry == EDatasmithOpenNurbsBrepTessellatedSource::UseUnrealNurbsTessellation)
 	{
 		// Ref. visitBRep
 		const FDatasmithOpenNurbsOptions& TessellationOptions = OpenNurbsOptions;
@@ -2988,7 +3022,7 @@ bool FOpenNurbsTranslatorImpl::TranslateBRep(ON_Brep* Brep, const ON_3dmObjectAt
 		return bool(Result);
 	}
 	else
-#endif
+#endif // CAD_LIBRARY
 	{
 		// .. Trying to load the mesh tessellated by Rhino
 		ON_SimpleArray<const ON_Mesh*> RenderMeshes;
@@ -2999,7 +3033,7 @@ bool FOpenNurbsTranslatorImpl::TranslateBRep(ON_Brep* Brep, const ON_3dmObjectAt
 		// Aborting because there is no mesh associated with the BRep
 		if (RenderMeshes.Count() == 0 && AnyMeshes.Count() == 0)
 		{
-			UE_LOG(LogDatasmithOpenNurbsTranslator, Warning, TEXT("No Render Mesh found for BRep with UUID %s."), *Name);
+			MissingRenderMeshes.Add(MeshElement->GetLabel());
 			return false;
 		}
 
@@ -3112,44 +3146,24 @@ void FOpenNurbsTranslatorImpl::SetOpenNurbsOptions(const FDatasmithOpenNurbsOpti
 {
 	OpenNurbsOptions = Options;
 	OpenNurbsOptionsHash = OpenNurbsOptions.GetHash();
-	
+
 	for (FOpenNurbsTranslatorImpl* ChildTranslator : ChildTranslators)
 	{
 		ChildTranslator->SetOpenNurbsOptions(OpenNurbsOptions);
 	}
 }
-#endif
 
 //////////////////////////////////////////////////////////////////////////
 // UDatasmithOpenNurbsTranslator
 //////////////////////////////////////////////////////////////////////////
 
-FDatasmithOpenNurbsTranslator::FDatasmithOpenNurbsTranslator()
-	: Translator(nullptr)
-{
-}
-
 void FDatasmithOpenNurbsTranslator::Initialize(FDatasmithTranslatorCapabilities& OutCapabilities)
 {
-#ifdef USE_OPENNURBS
 	OutCapabilities.SupportedFileFormats.Add(FFileFormatInfo{TEXT("3dm"), TEXT("Rhino file format")});
-#else
-	OutCapabilities.bIsEnabled = false;
-#endif
-}
-
-bool FDatasmithOpenNurbsTranslator::IsSourceSupported(const FDatasmithSceneSource& Source)
-{
-#ifdef USE_OPENNURBS
-	return Source.GetSourceFile().EndsWith(TEXT(".3dm"));
-#else
-	return false;
-#endif
 }
 
 bool FDatasmithOpenNurbsTranslator::LoadScene(TSharedRef<IDatasmithScene> OutScene)
 {
-#ifdef USE_OPENNURBS
 	const FString& Filename = GetSource().GetSourceFile();
 	FILE* FileHandle = ON::OpenFile(*Filename, L"rb");
 	if (!FileHandle)
@@ -3166,8 +3180,7 @@ bool FDatasmithOpenNurbsTranslator::LoadScene(TSharedRef<IDatasmithScene> OutSce
 	FString OutputPath = FPaths::ConvertRelativePathToFull(FPaths::Combine(FDatasmithOpenNurbsTranslatorModule::Get().GetTempDir(), TEXT("Cache"), GetSource().GetSceneName()));
 	IFileManager::Get().MakeDirectory(*OutputPath, true);
 	Translator->SetOutputPath(OutputPath);
-
-	Translator->SetOpenNurbsOptions(OpenNurbsImportOptions->Options);
+	Translator->SetOpenNurbsOptions(OpenNurbsOptions);
 	Translator->SetBaseOptions(BaseOptions);
 
 	ON_BinaryFile Archive(ON::archive_mode::read3dm, FileHandle);
@@ -3177,21 +3190,19 @@ bool FDatasmithOpenNurbsTranslator::LoadScene(TSharedRef<IDatasmithScene> OutSce
 	ON::CloseFile(FileHandle);
 
 	return bResult;
-#else
-	return false;
-#endif
 }
 
 void FDatasmithOpenNurbsTranslator::UnloadScene()
 {
-#ifdef USE_OPENNURBS
-	Translator.Reset();
-#endif
+	if (Translator)
+	{
+		Translator->ShowMessageLog(GetSource().GetSourceFile());
+		Translator.Reset();
+	}
 }
 
 bool FDatasmithOpenNurbsTranslator::LoadStaticMesh(const TSharedRef<IDatasmithMeshElement> MeshElement, FDatasmithMeshElementPayload& OutMeshPayload)
 {
-#ifdef USE_OPENNURBS
 	if ( TOptional< FMeshDescription > Mesh = Translator->GetMeshDescription( MeshElement ) )
 	{
 		OutMeshPayload.LodMeshes.Add(MoveTemp(Mesh.GetValue()));
@@ -3209,49 +3220,43 @@ bool FDatasmithOpenNurbsTranslator::LoadStaticMesh(const TSharedRef<IDatasmithMe
 				CoreTechData->RawData = MoveTemp(ByteArray);
 				CoreTechData->SceneParameters.ModelCoordSys = uint8(FDatasmithUtils::EModelCoordSystem::ZUp_RightHanded);
 				CoreTechData->SceneParameters.ScaleFactor = Translator->GetScalingFactor();
-				CoreTechData->SceneParameters.MetricUnit = Translator->GetMetricUnit(); 
-				CoreTechData->LastTessellationOptions = GetCommonTessellationOptions();
+				CoreTechData->SceneParameters.MetricUnit = Translator->GetMetricUnit();
+				CoreTechData->LastTessellationOptions = OpenNurbsOptions;
 				OutMeshPayload.AdditionalData.Add(CoreTechData);
 			}
 		}
-#endif
+#endif // CAD_LIBRARY
 	}
 
 	return OutMeshPayload.LodMeshes.Num() > 0;
-#else
-	return false;
-#endif
 }
 
 void FDatasmithOpenNurbsTranslator::SetSceneImportOptions(TArray<TStrongObjectPtr<UObject>>& Options)
 {
-#ifdef USE_OPENNURBS
 	for (TStrongObjectPtr<UObject>& Option : Options)
 	{
 		if (UDatasmithImportOptions* DatasmithOptions = Cast<UDatasmithImportOptions>(Option.Get()))
 		{
 			BaseOptions = DatasmithOptions->BaseOptions;
 		}
-		if (UDatasmithOpenNurbsImportOptions* OpenNurbsOptions = Cast<UDatasmithOpenNurbsImportOptions>(Option.Get()))
+		if (UDatasmithOpenNurbsImportOptions* OpenNurbsOptionsObj = Cast<UDatasmithOpenNurbsImportOptions>(Option.Get()))
 		{
-			OpenNurbsImportOptions.Reset(OpenNurbsOptions);
+			OpenNurbsOptions = OpenNurbsOptionsObj->Options;
 		}
 	}
 
 	if (Translator)
 	{
-		Translator->SetOpenNurbsOptions(OpenNurbsImportOptions->Options);
+		Translator->SetOpenNurbsOptions(OpenNurbsOptions);
 		Translator->SetBaseOptions(BaseOptions);
 	}
-#endif
 }
 
 void FDatasmithOpenNurbsTranslator::GetSceneImportOptions(TArray<TStrongObjectPtr<UObject>>& Options)
 {
-	if (!OpenNurbsImportOptions.IsValid())
-	{
-		OpenNurbsImportOptions = Datasmith::MakeOptions<UDatasmithOpenNurbsImportOptions>();
-	}
-
-	Options.Add(OpenNurbsImportOptions);
+	Options.Add(Datasmith::MakeOptions<UDatasmithOpenNurbsImportOptions>());
 }
+
+#undef LOCTEXT_NAMESPACE // "DatasmithOpenNurbsTranslator"
+
+#endif // USE_OPENNURBS
diff --git a/Engine/Plugins/Enterprise/DatasmithCADImporter/Source/DatasmithOpenNurbsTranslator/Private/DatasmithOpenNurbsTranslator.h b/Engine/Plugins/Enterprise/DatasmithCADImporter/Source/DatasmithOpenNurbsTranslator/Private/DatasmithOpenNurbsTranslator.h
--- a/Engine/Plugins/Enterprise/DatasmithCADImporter/Source/DatasmithOpenNurbsTranslator/Private/DatasmithOpenNurbsTranslator.h
+++ b/Engine/Plugins/Enterprise/DatasmithCADImporter/Source/DatasmithOpenNurbsTranslator/Private/DatasmithOpenNurbsTranslator.h
@@ -14,15 +14,16 @@ class FOpenNurbsTranslatorImpl;
 class FDatasmithOpenNurbsTranslator : public IDatasmithTranslator
 {
 public:
-	FDatasmithOpenNurbsTranslator();
+	virtual FName GetFName() const override { return "DatasmithOpenNurbsTranslator"; }
 
-	virtual FName GetFName() const override { return "DatasmithOpenNurbsTranslator"; };
+#ifndef USE_OPENNURBS
+	virtual void Initialize(FDatasmithTranslatorCapabilities& OutCapabilities) override { OutCapabilities.bIsEnabled = false; }
+#else // USE_OPENNURBS
 
 	virtual void Initialize(FDatasmithTranslatorCapabilities& OutCapabilities) override;
 
-	virtual bool IsSourceSupported(const FDatasmithSceneSource& Source) override;
-
 	virtual bool LoadScene(TSharedRef<IDatasmithScene> OutScene) override;
+
 	virtual void UnloadScene() override;
 
 	virtual bool LoadStaticMesh(const TSharedRef<IDatasmithMeshElement> MeshElement, FDatasmithMeshElementPayload& OutMeshPayload) override;
@@ -31,20 +32,14 @@ public:
 
 	virtual void GetSceneImportOptions(TArray<TStrongObjectPtr<UObject>>& Options) override;
 
-protected:
-	const FDatasmithTessellationOptions& GetCommonTessellationOptions()
-	{
-		return OpenNurbsImportOptions->Options;
-	}
-
 private:
-
 	TSharedPtr<FOpenNurbsTranslatorImpl> Translator;
 
 	// Temporarily store this here for UE-81278 so that we can trigger the recreation of
 	// static meshes if we're reimporting with new materials that haven't been assigned yet
 	FDatasmithImportBaseOptions BaseOptions;
 
-	TStrongObjectPtr<UDatasmithOpenNurbsImportOptions> OpenNurbsImportOptions;
+	FDatasmithOpenNurbsOptions OpenNurbsOptions;
+#endif // USE_OPENNURBS
 };
 
diff --git a/Engine/Plugins/Enterprise/DatasmithCADImporter/Source/DatasmithOpenNurbsTranslator/Private/RhinoCoretechWrapper.cpp b/Engine/Plugins/Enterprise/DatasmithCADImporter/Source/DatasmithOpenNurbsTranslator/Private/RhinoCoretechWrapper.cpp
--- a/Engine/Plugins/Enterprise/DatasmithCADImporter/Source/DatasmithOpenNurbsTranslator/Private/RhinoCoretechWrapper.cpp
+++ b/Engine/Plugins/Enterprise/DatasmithCADImporter/Source/DatasmithOpenNurbsTranslator/Private/RhinoCoretechWrapper.cpp
@@ -3,9 +3,8 @@
 #include "RhinoCoretechWrapper.h"
 
 
-#ifdef CAD_LIBRARY
+#if defined(CAD_LIBRARY) && defined(USE_OPENNURBS)
 #include "CoreTechHelper.h"
-//#include "TessellationHelper.h"
 
 #pragma warning(push)
 #pragma warning(disable:4265)
@@ -238,7 +237,7 @@ void BRepToKernelIOBodyTranslator::CreateCTFace_internal(const ON_BrepFace& Face
 
 			BrepTrimToCoedge[Trim.m_trim_index] = NewCoedge;
 
-			// Find a Trim that use this edge, that is not current Trim 
+			// Find a Trim that use this edge, that is not current Trim
 			// If the Trim has been converted into a coedge, link both Trims
 			for (int32 Index = 0; Index < on_edge->m_ti.Count(); ++Index)
 			{
@@ -475,4 +474,4 @@ TSharedPtr<FRhinoCoretechWrapper> FRhinoCoretechWrapper::GetSharedSession(double
 	return Session;
 }
 
-#endif
+#endif // defined(CAD_LIBRARY) && defined(USE_OPENNURBS)
diff --git a/Engine/Plugins/Enterprise/DatasmithCADImporter/Source/DatasmithOpenNurbsTranslator/Private/RhinoCoretechWrapper.h b/Engine/Plugins/Enterprise/DatasmithCADImporter/Source/DatasmithOpenNurbsTranslator/Private/RhinoCoretechWrapper.h
--- a/Engine/Plugins/Enterprise/DatasmithCADImporter/Source/DatasmithOpenNurbsTranslator/Private/RhinoCoretechWrapper.h
+++ b/Engine/Plugins/Enterprise/DatasmithCADImporter/Source/DatasmithOpenNurbsTranslator/Private/RhinoCoretechWrapper.h
@@ -2,7 +2,7 @@
 
 #pragma once
 
-#ifdef CAD_LIBRARY
+#if defined(CAD_LIBRARY) && defined(USE_OPENNURBS)
 #include "CoreMinimal.h"
 
 #include "CTSession.h"
@@ -59,4 +59,4 @@ protected:
 	static TWeakPtr<FRhinoCoretechWrapper> SharedSession;
 };
 
-#endif
\ No newline at end of file
+#endif // defined(CAD_LIBRARY) && defined(USE_OPENNURBS)
\ No newline at end of file
diff --git a/Engine/Plugins/Enterprise/DatasmithCADImporter/Source/DatasmithPLMXMLTranslator/Private/DatasmithPlmXmlTranslator.cpp b/Engine/Plugins/Enterprise/DatasmithCADImporter/Source/DatasmithPLMXMLTranslator/Private/DatasmithPlmXmlTranslator.cpp
new file mode 100644
--- /dev/null
+++ b/Engine/Plugins/Enterprise/DatasmithCADImporter/Source/DatasmithPLMXMLTranslator/Private/DatasmithPlmXmlTranslator.cpp
@@ -0,0 +1,71 @@
+// Copyright Epic Games, Inc. All Rights Reserved.
+
+#include "DatasmithPlmXmlTranslator.h"
+#include "DatasmithPlmXmlImporter.h"
+#include "DatasmithPlmXmlTranslatorModule.h"
+
+#include "DatasmithSceneFactory.h"
+#include "DatasmithSceneSource.h"
+#include "DatasmithImportOptions.h"
+#include "IDatasmithSceneElements.h"
+
+void FDatasmithPlmXmlTranslator::Initialize(FDatasmithTranslatorCapabilities& OutCapabilities)
+{
+	OutCapabilities.bIsEnabled = true;
+	OutCapabilities.bParallelLoadStaticMeshSupported = true;
+
+	TArray<FFileFormatInfo>& Formats = OutCapabilities.SupportedFileFormats;
+    Formats.Emplace(TEXT("plmxml"), TEXT("PML(Product Lifecycle Management) XML"));
+}
+
+bool FDatasmithPlmXmlTranslator::LoadScene(TSharedRef<IDatasmithScene> OutScene)
+{
+	OutScene->SetHost(TEXT("PlmXmlTranslator"));
+	OutScene->SetProductName(TEXT("PlmXml"));
+
+    Importer = MakeUnique<FDatasmithPlmXmlImporter>(OutScene);
+
+	const FString& FilePath = GetSource().GetSourceFile();
+	if (!Importer->OpenFile(FilePath, GetSource(), CommonTessellationOptionsPtr->Options))
+	{
+		return false;
+	}
+
+	return true;
+}
+
+void FDatasmithPlmXmlTranslator::UnloadScene()
+{
+	Importer->UnloadScene();
+	Importer.Reset();
+}
+
+bool FDatasmithPlmXmlTranslator::LoadStaticMesh(const TSharedRef<IDatasmithMeshElement> MeshElement, FDatasmithMeshElementPayload& OutMeshPayload)
+{
+	if (ensure(Importer.IsValid()))
+	{
+		return Importer->LoadStaticMesh(MeshElement, OutMeshPayload);
+	}
+
+	return false;
+}
+
+void FDatasmithPlmXmlTranslator::GetSceneImportOptions(TArray<TStrongObjectPtr<UObject>>& Options)
+{
+	if (!CommonTessellationOptionsPtr.IsValid())
+	{
+		CommonTessellationOptionsPtr = Datasmith::MakeOptions<UDatasmithCommonTessellationOptions>();
+	}
+	Options.Add(CommonTessellationOptionsPtr);
+}
+
+void FDatasmithPlmXmlTranslator::SetSceneImportOptions(TArray<TStrongObjectPtr<UObject>>& Options)
+{
+	for (const TStrongObjectPtr<UObject>& OptionPtr : Options)
+	{
+		if (UDatasmithCommonTessellationOptions* TessellationOptionsObject = Cast<UDatasmithCommonTessellationOptions>(OptionPtr.Get()))
+		{
+			CommonTessellationOptionsPtr.Reset(TessellationOptionsObject);
+		}
+	}
+}
