Trying to fix the case issue on linux
This might override the server meta data for the files so it's worth trying.

Merging

//UE4/Dev-Tools-Staging/Engine/Plugins/Enterprise/DataprepEditor/...
to //UE4/Dev-VirtualProduction/Engine/Plugins/Enterprise/DataprepEditor/...

diff --git a/Engine/Plugins/Enterprise/DataprepEditor/Config/BaseDataprepEditor.ini b/Engine/Plugins/Enterprise/DataprepEditor/Config/BaseDataprepEditor.ini
new file mode 100644
--- /dev/null
+++ b/Engine/Plugins/Enterprise/DataprepEditor/Config/BaseDataprepEditor.ini
@@ -0,0 +1,7 @@
+[CoreRedirects]
+
++ClassRedirects=(OldName="/Script/DataprepEditor.K2Node_DataprepProducer", NewName="/Script/DataprepCore.K2Node_DataprepProducer")
++ClassRedirects=(OldName="DataprepRemoveObjectsOperation", NewName="DataprepDeleteObjectsOperation")
++PropertyRedirects=(OldName="DataprepSetMaterialOperation.MaterialSubstitute",NewName="DataprepSetMaterialOperation.Material")
++PropertyRedirects=(OldName="DataprepSetMeshOperation.MeshSubstitute",NewName="DataprepSetMeshOperation.StaticMesh")
++EnumRedirects=(OldName="EDataprepFloatMatchType",ValueChanges=(("LessThen","LessThan"),("GreatherThen","GreatherThan")))
diff --git a/Engine/Plugins/Enterprise/DataprepEditor/Source/DataprepCore/Private/Blueprint/K2Node_DataprepActionCore.cpp b/Engine/Plugins/Enterprise/DataprepEditor/Source/DataprepCore/Private/Blueprint/K2Node_DataprepActionCore.cpp
new file mode 100644
--- /dev/null
+++ b/Engine/Plugins/Enterprise/DataprepEditor/Source/DataprepCore/Private/Blueprint/K2Node_DataprepActionCore.cpp
@@ -0,0 +1,15 @@
+// Copyright 1998-2019 Epic Games, Inc. All Rights Reserved.
+
+#include "Blueprint/K2Node_DataprepActionCore.h"
+
+// Dataprep includes
+#include "DataprepActionAsset.h"
+
+void UK2Node_DataprepActionCore::CreateDataprepActionAsset()
+{
+	if ( !DataprepActionAsset )
+	{
+		DataprepActionAsset = NewObject< UDataprepActionAsset >( this, UDataprepActionAsset::StaticClass(), NAME_None, RF_Transactional );
+		DataprepActionAsset->SetLabel( *GetNodeTitle(ENodeTitleType::EditableTitle).ToString() );
+	}
+}
diff --git a/Engine/Plugins/Enterprise/DataprepEditor/Source/DataprepCore/Private/DataprepAsset.cpp b/Engine/Plugins/Enterprise/DataprepEditor/Source/DataprepCore/Private/DataprepAsset.cpp
new file mode 100644
--- /dev/null
+++ b/Engine/Plugins/Enterprise/DataprepEditor/Source/DataprepCore/Private/DataprepAsset.cpp
@@ -0,0 +1,321 @@
+// Copyright 1998-2019 Epic Games, Inc. All Rights Reserved.
+
+#include "DataprepAsset.h"
+
+#include "Blueprint/K2Node_DataprepActionCore.h"
+#include "Blueprint/K2Node_DataprepProducer.h"
+#include "DataprepContentConsumer.h"
+#include "DataprepContentProducer.h"
+#include "DataprepRecipe.h"
+#include "DataprepActionAsset.h"
+#include "DataprepCoreLogCategory.h"
+#include "DataprepCorePrivateUtils.h"
+#include "DataprepCoreUtils.h"
+#include "DataprepParameterizableObject.h"
+#include "Parameterization/DataprepParameterization.h"
+
+#include "AssetRegistryModule.h"
+#include "BlueprintNodeBinder.h"
+#include "BlueprintNodeSpawner.h"
+#include "Kismet2/BlueprintEditorUtils.h"
+#include "Kismet2/KismetEditorUtilities.h"
+#include "UObject/UObjectGlobals.h"
+
+#ifdef WITH_EDITOR
+#include "Editor.h"
+#endif //WITH_EDITOR
+
+#define LOCTEXT_NAMESPACE "DataprepAsset"
+
+// UDataprepAsset =================================================================
+
+UDataprepAsset::UDataprepAsset()
+{
+	// Temp code for the nodes development
+	DataprepRecipeBP = nullptr;
+	StartNode = nullptr;
+	// end of temp code for nodes development
+}
+
+void UDataprepAsset::PostLoad()
+{
+	UDataprepAssetInterface::PostLoad();
+
+	check(DataprepRecipeBP);
+	DataprepRecipeBP->OnChanged().AddUObject( this, &UDataprepAsset::OnDataprepBlueprintChanged );
+
+	// Move content of deprecated properties to the corresponding new ones.
+	if(HasAnyFlags(RF_WasLoaded))
+	{
+		bool bMarkDirty = false;
+		if(Producers_DEPRECATED.Num() > 0)
+		{
+			Inputs->AssetProducers.Reserve(Producers_DEPRECATED.Num());
+
+			while(Producers_DEPRECATED.Num() > 0)
+			{
+				if(Inputs->AddAssetProducer( Producers_DEPRECATED.Pop(false) ) == INDEX_NONE)
+				{
+					// #ueent_todo Log message a producer was not properly restored
+				}
+			}
+
+			Producers_DEPRECATED.Empty();
+			bMarkDirty = true;
+		}
+
+		if(Consumer_DEPRECATED)
+		{
+			Output = Consumer_DEPRECATED;
+			Consumer_DEPRECATED = nullptr;
+			bMarkDirty = true;
+		}
+
+		// Most likely a Dataprep asset from 4.23
+		if(StartNode == nullptr)
+		{
+			UEdGraph* PipelineGraph = FBlueprintEditorUtils::FindEventGraph(DataprepRecipeBP);
+			check( PipelineGraph );
+
+			for( UEdGraphNode* GraphNode : PipelineGraph->Nodes )
+			{
+				StartNode = Cast<UK2Node_DataprepProducer>(GraphNode);
+				if( StartNode )
+				{
+					break;
+				}
+			}
+
+			// This Dataprep asset was never opened in the editor
+			if( StartNode == nullptr )
+			{
+				IBlueprintNodeBinder::FBindingSet Bindings;
+				StartNode = UBlueprintNodeSpawner::Create<UK2Node_DataprepProducer>()->Invoke( PipelineGraph, Bindings, FVector2D(-100,0) );
+				check(Cast<UK2Node_DataprepProducer>(StartNode));
+
+				DataprepRecipeBP->MarkPackageDirty();
+			}
+
+			UpdateActions();
+			bMarkDirty = true;
+		}
+
+		if ( !Parameterization )
+		{
+			Parameterization = NewObject<UDataprepParameterization>( this, FName(), RF_Public | RF_Transactional );
+			bMarkDirty = true;
+		}
+
+		// Mark the asset as dirty to indicate asset's properties have changed
+		if(bMarkDirty)
+		{
+			const FText AssetName = FText::FromString( GetName() );
+			const FText WarningMessage = FText::Format( LOCTEXT( "DataprepAssetOldVersion", "{0} is from an old version and has been updated. Please save asset to complete update."), AssetName );
+			const FText NotificationText = FText::Format( LOCTEXT( "DataprepAssetOldVersionNotif", "{0} is from an old version and has been updated."), AssetName );
+			DataprepCorePrivateUtils::LogMessage( EMessageSeverity::Warning, WarningMessage, NotificationText );
+
+			GetOutermost()->SetDirtyFlag(true);
+		}
+	}
+}
+
+bool UDataprepAsset::Rename(const TCHAR* NewName/* =nullptr */, UObject* NewOuter/* =nullptr */, ERenameFlags Flags/* =REN_None */)
+{
+	bool bWasRename = Super::Rename( NewName, NewOuter, Flags );
+	if ( bWasRename )
+	{
+		if ( Parameterization )
+		{
+			bWasRename &= Parameterization->OnAssetRename( Flags );
+		}
+
+		if ( DataprepRecipeBP && bWasRename )
+		{
+			// There shouldn't be a blueprint depending on us. Should be ok to just rename the generated class
+			bWasRename &= DataprepRecipeBP->RenameGeneratedClasses( NewName, NewOuter, Flags );
+		}
+	}
+
+	return bWasRename;
+}
+
+bool UDataprepAsset::CreateBlueprint()
+{
+	// Begin: Temp code for the nodes development
+	const FString DesiredName = GetName() + TEXT("_Recipe");
+	FName BlueprintName = MakeUniqueObjectName( GetOutermost(), UBlueprint::StaticClass(), *DesiredName );
+
+	DataprepRecipeBP = FKismetEditorUtilities::CreateBlueprint( UDataprepRecipe::StaticClass(), this, BlueprintName, BPTYPE_Normal, UBlueprint::StaticClass(), UBlueprintGeneratedClass::StaticClass() );
+	check( DataprepRecipeBP );
+
+	// This blueprint is not the asset of the package
+	DataprepRecipeBP->ClearFlags( RF_Standalone );
+
+	FAssetRegistryModule::AssetCreated( DataprepRecipeBP );
+
+	// Create the start node of the Blueprint
+	UEdGraph* PipelineGraph = FBlueprintEditorUtils::FindEventGraph(DataprepRecipeBP);
+	check( PipelineGraph );
+
+	UEdGraph* EventGraph = FBlueprintEditorUtils::FindEventGraph(DataprepRecipeBP);
+	IBlueprintNodeBinder::FBindingSet Bindings;
+
+	StartNode = UBlueprintNodeSpawner::Create<UK2Node_DataprepProducer>()->Invoke( EventGraph, Bindings, FVector2D(-100,0) );
+	check(Cast<UK2Node_DataprepProducer>(StartNode));
+
+	DataprepRecipeBP->MarkPackageDirty();
+
+	DataprepRecipeBP->OnChanged().AddUObject( this, &UDataprepAsset::OnDataprepBlueprintChanged );
+	// End: Temp code for the nodes development
+
+	MarkPackageDirty();
+
+	return true;
+}
+
+bool UDataprepAsset::CreateParameterization()
+{
+	if( !Parameterization )
+	{
+		Parameterization = NewObject<UDataprepParameterization>( this, FName(), RF_Public | RF_Transactional );
+		MarkPackageDirty();
+		return true;
+	}
+
+	return false;
+}
+
+void UDataprepAsset::ExecuteRecipe(const TSharedPtr<FDataprepActionContext>& InActionsContext)
+{
+	ExecuteRecipe_Internal(	InActionsContext, ActionAssets );
+}
+
+TArray<UDataprepActionAsset*> UDataprepAsset::GetCopyOfActions(TMap<UObject*,UObject*>& OutOriginalToCopy) const
+{
+	TArray<UDataprepActionAsset*> CopyOfActionAssets;
+	CopyOfActionAssets.Reserve( ActionAssets.Num() );
+	for ( UDataprepActionAsset* ActionAsset : ActionAssets )
+	{
+		FObjectDuplicationParameters DuplicationParameter( ActionAsset, GetTransientPackage() );
+		DuplicationParameter.CreatedObjects = &OutOriginalToCopy;
+
+		UDataprepActionAsset* CopyOfAction = static_cast<UDataprepActionAsset*>( StaticDuplicateObjectEx( DuplicationParameter ) );
+		check( CopyOfAction );
+
+		OutOriginalToCopy.Add( ActionAsset, CopyOfAction );
+		CopyOfActionAssets.Add( CopyOfAction );
+	}
+
+	return CopyOfActionAssets;
+}
+
+UObject* UDataprepAsset::GetParameterizationObject()
+{
+	return Parameterization->GetDefaultObject();
+}
+
+void UDataprepAsset::BindObjectPropertyToParameterization(UDataprepParameterizableObject* Object, const TArray<struct FDataprepPropertyLink>& InPropertyChain, const FName& Name)
+{
+	bool bPassConditionCheck = false;
+
+	if ( InPropertyChain.Num() > 0 )
+	{
+		// Validate that the object is part of this asset
+		UObject* Outer = Object;
+		while ( Outer && !bPassConditionCheck )
+		{
+			Outer =  Outer->GetOuter();
+			bPassConditionCheck = Outer == this;
+		}
+	}
+
+	if ( bPassConditionCheck )
+	{
+		Parameterization->BindObjectProperty( Object, InPropertyChain, Name );
+	}
+
+}
+
+bool UDataprepAsset::IsObjectPropertyBinded(UDataprepParameterizableObject* Object, const TArray<struct FDataprepPropertyLink>& InPropertyChain) const
+{
+	return Parameterization->IsObjectPropertyBinded( Object, InPropertyChain );
+}
+
+FName UDataprepAsset::GetNameOfParameterForObjectProperty(UDataprepParameterizableObject* Object, const TArray<struct FDataprepPropertyLink>& InPropertyChain) const
+{
+	return Parameterization->GetNameOfParameterForObjectProperty( Object, InPropertyChain );
+}
+
+void UDataprepAsset::RemoveObjectPropertyFromParameterization(UDataprepParameterizableObject* Object, const TArray<struct FDataprepPropertyLink>& InPropertyChain)
+{
+	Parameterization->RemoveBindedObjectProperty( Object, InPropertyChain );
+}
+
+void UDataprepAsset::GetExistingParameterNamesForType(UProperty* Property, bool bIsDescribingFullProperty, TSet<FString>& OutValidExistingNames, TSet<FString>& OutInvalidNames) const
+{
+	Parameterization->GetExistingParameterNamesForType( Property, bIsDescribingFullProperty, OutValidExistingNames, OutInvalidNames );
+}
+
+void UDataprepAsset::OnDataprepBlueprintChanged( UBlueprint* InBlueprint )
+{
+	if(InBlueprint == DataprepRecipeBP)
+	{
+		UpdateActions();
+		OnChanged.Broadcast( FDataprepAssetChangeType::RecipeModified );
+	}
+}
+
+void UDataprepAsset::UpdateActions()
+{
+	ActionAssets.Empty(ActionAssets.Num());
+
+	UEdGraphPin* NodeOutPin= StartNode->FindPin(UEdGraphSchema_K2::PN_Then, EGPD_Output);
+	if( NodeOutPin && NodeOutPin->LinkedTo.Num() > 0 )
+	{
+		TSet<UEdGraphNode*> ActionNodesVisited;
+
+		for( UEdGraphPin* NextNodeInPin = NodeOutPin->LinkedTo[0]; NextNodeInPin != nullptr ; )
+		{
+			UEdGraphNode* NextNode = NextNodeInPin->GetOwningNode();
+
+			uint32 NodeHash = GetTypeHash( NextNode );
+			// Break the loop if the node had already been visited
+			if( ActionNodesVisited.FindByHash( NodeHash, NextNode ) )
+			{
+				break;
+			}
+			else
+			{
+				ActionNodesVisited.AddByHash( NodeHash, NextNode );
+			}
+
+			if(UK2Node_DataprepActionCore* ActionNode = Cast<UK2Node_DataprepActionCore>(NextNode))
+			{
+				if( UDataprepActionAsset* DataprepAction = ActionNode->GetDataprepAction() )
+				{
+					ActionAssets.Add( DataprepAction );
+				}
+			}
+
+			// Look for the next node
+			NodeOutPin = NextNode->FindPin( UEdGraphSchema_K2::PN_Then, EGPD_Output );
+
+			if ( !NodeOutPin )
+			{
+				// If we couldn't find a then pin try to get the first output pin as a fallback
+				for ( UEdGraphPin* Pin : NextNode->Pins )
+				{
+					if ( Pin && Pin->PinType.PinCategory == UEdGraphSchema_K2::PC_Exec && Pin->Direction == EGPD_Output )
+					{
+						NodeOutPin = Pin;
+						break;
+					}
+				}
+			}
+
+			NextNodeInPin = NodeOutPin ? ( NodeOutPin->LinkedTo.Num() > 0 ? NodeOutPin->LinkedTo[0] : nullptr ) : nullptr;
+		}
+	}
+}
+
+#undef LOCTEXT_NAMESPACE
diff --git a/Engine/Plugins/Enterprise/DataprepEditor/Source/DataprepCore/Private/DataprepCoreLogCategory.cpp b/Engine/Plugins/Enterprise/DataprepEditor/Source/DataprepCore/Private/DataprepCoreLogCategory.cpp
new file mode 100644
--- /dev/null
+++ b/Engine/Plugins/Enterprise/DataprepEditor/Source/DataprepCore/Private/DataprepCoreLogCategory.cpp
@@ -0,0 +1,5 @@
+// Copyright 1998-2019 Epic Games, Inc. All Rights Reserved.
+
+#include "DataprepCoreLogCategory.h"
+
+DEFINE_LOG_CATEGORY(LogDataprepCore);
diff --git a/Engine/Plugins/Enterprise/DataprepEditor/Source/DataprepCore/Private/DataprepCoreLogCategory.h b/Engine/Plugins/Enterprise/DataprepEditor/Source/DataprepCore/Private/DataprepCoreLogCategory.h
new file mode 100644
--- /dev/null
+++ b/Engine/Plugins/Enterprise/DataprepEditor/Source/DataprepCore/Private/DataprepCoreLogCategory.h
@@ -0,0 +1,7 @@
+// Copyright 1998-2019 Epic Games, Inc. All Rights Reserved.
+
+#pragma once
+
+#include "Logging/LogMacros.h"
+
+DECLARE_LOG_CATEGORY_EXTERN(LogDataprepCore, Log, All);
diff --git a/Engine/Plugins/Enterprise/DataprepEditor/Source/DataprepCore/Private/DataprepCoreModule.cpp b/Engine/Plugins/Enterprise/DataprepEditor/Source/DataprepCore/Private/DataprepCoreModule.cpp
new file mode 100644
--- /dev/null
+++ b/Engine/Plugins/Enterprise/DataprepEditor/Source/DataprepCore/Private/DataprepCoreModule.cpp
@@ -0,0 +1,24 @@
+// Copyright 1998-2019 Epic Games, Inc. All Rights Reserved.
+
+#include "DataprepCorePrivateUtils.h"
+
+#include "Misc/PackageName.h"
+#include "Modules/ModuleManager.h"
+
+class FDataprepCoreModule : public IModuleInterface
+{
+public:
+	virtual void StartupModule() override
+	{
+		// Register mount point for Dataprep core library root package folder
+		FPackageName::RegisterMountPoint( DataprepCorePrivateUtils::GetRootPackagePath() + TEXT("/"), DataprepCorePrivateUtils::GetRootTemporaryDir()	);
+	}
+
+	virtual void ShutdownModule() override
+	{
+		// Unregister mount point for Dataprep core library root package folder
+		FPackageName::UnRegisterMountPoint( DataprepCorePrivateUtils::GetRootPackagePath() + TEXT("/"), DataprepCorePrivateUtils::GetRootTemporaryDir() );
+	}
+};
+
+IMPLEMENT_MODULE( FDataprepCoreModule, DataprepCore )
diff --git a/Engine/Plugins/Enterprise/DataprepEditor/Source/DataprepCore/Private/DataprepCorePrivateUtils.cpp b/Engine/Plugins/Enterprise/DataprepEditor/Source/DataprepCore/Private/DataprepCorePrivateUtils.cpp
new file mode 100644
--- /dev/null
+++ b/Engine/Plugins/Enterprise/DataprepEditor/Source/DataprepCore/Private/DataprepCorePrivateUtils.cpp
@@ -0,0 +1,214 @@
+// Copyright 1998-2019 Epic Games, Inc. All Rights Reserved.
+
+#include "DataprepCorePrivateUtils.h"
+
+#include "DataprepAsset.h"
+#include "DataprepCoreUtils.h"
+#include "IDataprepProgressReporter.h"
+
+#include "ActorEditorUtils.h"
+#include "AssetRegistryModule.h"
+#include "Engine/Blueprint.h"
+#include "Engine/StaticMesh.h"
+#include "Engine/World.h"
+#include "GameFramework/WorldSettings.h"
+#include "IMessageLogListing.h"
+#include "Kismet2/BlueprintEditorUtils.h"
+#include "MeshDescription.h"
+#include "MessageLogModule.h"
+#include "Modules/ModuleManager.h"
+#include "StaticMeshAttributes.h"
+
+#define LOCTEXT_NAMESPACE "DataprepAsset"
+
+void DataprepCorePrivateUtils::DeleteRegisteredAsset(UObject* Asset)
+{
+	if(Asset != nullptr)
+	{
+		FDataprepCoreUtils::MoveToTransientPackage( Asset );
+
+		Asset->ClearFlags(RF_Standalone | RF_Public);
+		Asset->RemoveFromRoot();
+		Asset->MarkPendingKill();
+
+		FAssetRegistryModule::AssetDeleted( Asset ) ;
+	}
+}
+
+void DataprepCorePrivateUtils::GetActorsFromWorld(const UWorld* World, TArray<AActor*>& OutActors )
+{
+	if(World != nullptr)
+	{
+		int32 ActorsCount = 0;
+		for(ULevel* Level : World->GetLevels())
+		{
+			ActorsCount += Level->Actors.Num();
+		}
+
+		OutActors.Reserve( OutActors.Num() + ActorsCount );
+
+		for(ULevel* Level : World->GetLevels())
+		{
+			for( AActor* Actor : Level->Actors )
+			{
+				const bool bIsValidActor = Actor &&
+					!Actor->IsPendingKill() &&
+					Actor->IsEditable() &&
+					!Actor->IsTemplate() &&
+					!FActorEditorUtils::IsABuilderBrush(Actor) &&
+					!Actor->IsA(AWorldSettings::StaticClass());
+
+				if( bIsValidActor )
+				{
+					OutActors.Add( Actor  );
+				}
+			}
+		}
+	}
+}
+
+
+const FString& DataprepCorePrivateUtils::GetRootTemporaryDir()
+{
+	static FString RootTemporaryDir = FPaths::Combine(FPaths::ProjectIntermediateDir(), TEXT("DataprepTemp") );
+	return RootTemporaryDir;
+}
+
+const FString& DataprepCorePrivateUtils::GetRootPackagePath()
+{
+	static FString RootPackagePath( TEXT("/Engine/DataprepCore/Transient") );
+	return RootPackagePath;
+}
+
+void DataprepCorePrivateUtils::LogMessage( EMessageSeverity::Type Severity, const FText& Message, const FText& NotificationText )
+{
+	FMessageLogModule& MessageLogModule = FModuleManager::LoadModuleChecked<FMessageLogModule>("MessageLog");
+
+	TSharedPtr<IMessageLogListing> LogListing = MessageLogModule.GetLogListing( TEXT("DataprepCore") );
+	LogListing->SetLabel( LOCTEXT("MessageLogger", "Dataprep Core") );
+
+	LogListing->AddMessage( FTokenizedMessage::Create( Severity, Message ), /*bMirrorToOutputLog*/ true );
+
+	if( !NotificationText.IsEmpty() )
+	{
+		LogListing->NotifyIfAnyMessages( NotificationText, EMessageSeverity::Info);
+	}
+}
+
+void DataprepCorePrivateUtils::BuildStaticMeshes(TSet<UStaticMesh*>& StaticMeshes, TFunction<bool(UStaticMesh*)> ProgressFunction, bool bForceBuild)
+{
+	TArray<UStaticMesh*> BuiltMeshes;
+	BuiltMeshes.Reserve( StaticMeshes.Num() );
+
+	if(bForceBuild)
+	{
+		BuiltMeshes.Append( StaticMeshes.Array() );
+	}
+	else
+	{
+		for(UStaticMesh* StaticMesh : StaticMeshes)
+		{
+			if(StaticMesh && (!StaticMesh->RenderData.IsValid() || !StaticMesh->RenderData->IsInitialized()))
+			{
+				BuiltMeshes.Add( StaticMesh );
+			}
+		}
+	}
+
+	if(BuiltMeshes.Num() > 0)
+	{
+		// Start with the biggest mesh first to help balancing tasks on threads
+		BuiltMeshes.Sort(
+			[](const UStaticMesh& Lhs, const UStaticMesh& Rhs) 
+		{ 
+			int32 LhsVerticesNum = Lhs.IsMeshDescriptionValid(0) ? Lhs.GetMeshDescription(0)->Vertices().Num() : 0;
+			int32 RhsVerticesNum = Rhs.IsMeshDescriptionValid(0) ? Rhs.GetMeshDescription(0)->Vertices().Num() : 0;
+
+			return LhsVerticesNum > RhsVerticesNum;
+		}
+		);
+
+		TArray< TArray<FMeshBuildSettings> > StaticMeshesSettings;
+		StaticMeshesSettings.Reserve( BuiltMeshes.Num() );
+
+		//Cache the BuildSettings and update them before building the meshes.
+		for (UStaticMesh* StaticMesh : BuiltMeshes)
+		{
+			TArray<FStaticMeshSourceModel>& SourceModels = StaticMesh->GetSourceModels();
+			TArray<FMeshBuildSettings> BuildSettings;
+			BuildSettings.Reserve(SourceModels.Num());
+
+			for(int32 Index = 0; Index < SourceModels.Num(); ++Index)
+			{
+				FStaticMeshSourceModel& SourceModel = SourceModels[Index];
+
+				BuildSettings.Add( SourceModel.BuildSettings );
+
+				if(FMeshDescription* MeshDescription = StaticMesh->GetMeshDescription(Index))
+				{
+					FStaticMeshAttributes Attributes(*MeshDescription);
+					if(SourceModel.BuildSettings.DstLightmapIndex != -1)
+					{
+						TVertexInstanceAttributesConstRef<FVector2D> VertexInstanceUVs = Attributes.GetVertexInstanceUVs();
+						SourceModel.BuildSettings.bGenerateLightmapUVs = VertexInstanceUVs.IsValid() && VertexInstanceUVs.GetNumIndices() > SourceModel.BuildSettings.DstLightmapIndex;
+					}
+					else
+					{
+						SourceModel.BuildSettings.bGenerateLightmapUVs = false;
+					}
+
+					SourceModel.BuildSettings.bRecomputeNormals = !(Attributes.GetVertexInstanceNormals().IsValid() && Attributes.GetVertexInstanceNormals().GetNumIndices() > 0);
+					SourceModel.BuildSettings.bRecomputeTangents = false;
+					//SourceModel.BuildSettings.bBuildAdjacencyBuffer = false;
+					//SourceModel.BuildSettings.bBuildReversedIndexBuffer = false;
+				}
+			}
+
+			StaticMeshesSettings.Add(MoveTemp(BuildSettings));				
+		}
+
+		// Disable warnings from LogStaticMesh. Not useful
+		ELogVerbosity::Type PrevLogStaticMeshVerbosity = LogStaticMesh.GetVerbosity();
+		LogStaticMesh.SetVerbosity( ELogVerbosity::Error );
+
+		UStaticMesh::BatchBuild(BuiltMeshes, true, ProgressFunction);
+
+		// Restore LogStaticMesh verbosity
+		LogStaticMesh.SetVerbosity( PrevLogStaticMeshVerbosity );
+
+		for(int32 Index = 0; Index < BuiltMeshes.Num(); ++Index)
+		{
+			UStaticMesh* StaticMesh = BuiltMeshes[Index];
+			TArray<FMeshBuildSettings>& PrevBuildSettings = StaticMeshesSettings[Index];
+
+			TArray<FStaticMeshSourceModel>& SourceModels = StaticMesh->GetSourceModels();
+
+			for(int32 SourceModelIndex = 0; SourceModelIndex < SourceModels.Num(); ++SourceModelIndex)
+			{
+				SourceModels[SourceModelIndex].BuildSettings = PrevBuildSettings[SourceModelIndex];
+			}
+
+			if(FStaticMeshRenderData* RenderData = StaticMesh->RenderData.Get())
+			{
+				for ( FStaticMeshLODResources& LODResources : RenderData->LODResources )
+				{
+					LODResources.bHasColorVertexData = true;
+				}
+			}
+		}
+	}
+}
+
+void DataprepCorePrivateUtils::ClearAssets(const TArray<TWeakObjectPtr<UObject>>& Assets)
+{
+	for(const TWeakObjectPtr<UObject>& ObjectPtr : Assets)
+	{
+		if(UStaticMesh* StaticMesh = Cast<UStaticMesh>(ObjectPtr.Get()))
+		{
+			StaticMesh->PreEditChange( nullptr );
+			StaticMesh->RenderData.Reset();
+		}
+	}
+}
+
+#undef LOCTEXT_NAMESPACE
\ No newline at end of file
diff --git a/Engine/Plugins/Enterprise/DataprepEditor/Source/DataprepCore/Private/DataprepFactories.cpp b/Engine/Plugins/Enterprise/DataprepEditor/Source/DataprepCore/Private/DataprepFactories.cpp
new file mode 100644
--- /dev/null
+++ b/Engine/Plugins/Enterprise/DataprepEditor/Source/DataprepCore/Private/DataprepFactories.cpp
@@ -0,0 +1,120 @@
+// Copyright 1998-2019 Epic Games, Inc. All Rights Reserved.
+
+#include "DataprepFactories.h"
+
+#include "DataprepAsset.h"
+#include "DataprepAssetInstance.h"
+#include "DataprepContentConsumer.h"
+
+#include "AssetRegistryModule.h"
+#include "AssetTypeCategories.h"
+#include "UObject/UObjectHash.h"
+#include "UObject/UObjectIterator.h"
+
+UDataprepAssetFactory::UDataprepAssetFactory()
+{
+	SupportedClass = UDataprepAsset::StaticClass();
+
+	bCreateNew = true;
+	bText = false;
+	bEditorImport = false;
+}
+
+bool UDataprepAssetFactory::ShouldShowInNewMenu() const
+{
+	// If there is no consumer don't show this factory
+	TArray< UClass* > PotentialClasses;
+	GetDerivedClasses( UDataprepContentConsumer::StaticClass(), PotentialClasses, true );
+
+	for ( UClass* ChildClass : PotentialClasses )
+	{
+		if ( ChildClass && !ChildClass->HasAnyClassFlags( CLASS_CompiledFromBlueprint | CLASS_Deprecated | CLASS_NewerVersionExists | CLASS_Abstract ) && ChildClass->HasAllClassFlags( CLASS_Native ) )
+		{
+			return true;
+		}
+	}
+
+	return false;
+}
+
+UObject * UDataprepAssetFactory::FactoryCreateNew(UClass* InClass, UObject* InParent, FName InName, EObjectFlags Flags, UObject* Context, FFeedbackContext *Warn)
+{
+	check( InClass->IsChildOf( UDataprepAsset::StaticClass() ) );
+
+	// Find potential Consumer classes
+	TArray<UClass*> ConsumerClasses;
+	for( TObjectIterator< UClass > It ; It ; ++It )
+	{
+		UClass* CurrentClass = (*It);
+
+		if ( !CurrentClass->HasAnyClassFlags( CLASS_Abstract ) )
+		{
+			if( CurrentClass->IsChildOf( UDataprepContentConsumer::StaticClass() ) )
+			{
+				ConsumerClasses.Add( CurrentClass );
+			}
+		}
+	}
+
+	if(ConsumerClasses.Num() == 0)
+	{
+		return nullptr;
+	}
+
+	UDataprepAsset* DataprepAsset = NewObject<UDataprepAsset>(InParent, InClass, InName, Flags | RF_Transactional);
+	check(DataprepAsset);
+
+	// Initialize Dataprep asset's consumer
+	if(ConsumerClasses.Num() == 1)
+	{
+		DataprepAsset->SetConsumer( ConsumerClasses[0], /* bNotifyChanges = */ false );
+	}
+	else
+	{
+		// #ueent_todo: Propose user to choose from the list of Consumers.
+		DataprepAsset->SetConsumer( ConsumerClasses[0], /* bNotifyChanges = */ false );
+	}
+	check( DataprepAsset->GetConsumer() );
+
+	// Temp code for the nodes development
+	// Initialize Dataprep asset's blueprint
+	DataprepAsset->CreateBlueprint();
+	// end of temp code for nodes development
+
+	DataprepAsset->CreateParameterization();
+
+	FAssetRegistryModule::AssetCreated( DataprepAsset );
+	DataprepAsset->MarkPackageDirty();
+
+	return DataprepAsset;
+}
+
+UDataprepAssetInstanceFactory::UDataprepAssetInstanceFactory()
+{
+	SupportedClass = UDataprepAssetInstance::StaticClass();
+
+	bCreateNew = false;
+	bText = false;
+	bEditorImport = false;
+}
+
+UObject* UDataprepAssetInstanceFactory::FactoryCreateNew(UClass* InClass, UObject* InParent, FName InName, EObjectFlags Flags, UObject* Context, FFeedbackContext* Warn)
+{
+	if(UDataprepAssetInterface* DataprepAssetParent = Cast<UDataprepAssetInterface>(InitialParent))
+	{
+		UDataprepAssetInstance* DataprepAssetInstance = NewObject<UDataprepAssetInstance>(InParent, InClass, InName, Flags);
+
+		if(DataprepAssetInstance && DataprepAssetParent->GetConsumer())
+		{
+			if(DataprepAssetInstance->SetParent(DataprepAssetParent, /* bNotifyChanges = */ false))
+			{
+				FAssetRegistryModule::AssetCreated( DataprepAssetInstance );
+				DataprepAssetInstance->MarkPackageDirty();
+
+				return DataprepAssetInstance;
+			}
+		}
+	}
+
+	return nullptr;
+}
\ No newline at end of file
diff --git a/Engine/Plugins/Enterprise/DataprepEditor/Source/DataprepCore/Private/DataprepRecipe.cpp b/Engine/Plugins/Enterprise/DataprepEditor/Source/DataprepCore/Private/DataprepRecipe.cpp
new file mode 100644
--- /dev/null
+++ b/Engine/Plugins/Enterprise/DataprepEditor/Source/DataprepCore/Private/DataprepRecipe.cpp
@@ -0,0 +1,106 @@
+// Copyright 1998-2019 Epic Games, Inc. All Rights Reserved.
+
+#include "DataprepRecipe.h"
+
+#include "ActorEditorUtils.h"
+#include "Engine/World.h"
+#include "Engine/StaticMesh.h"
+#include "EngineUtils.h"
+#include "GameFramework/Actor.h"
+#include "Materials/Material.h"
+#include "UObject/SoftObjectPath.h"
+
+#define LOCTEXT_NAMESPACE "DataprepRecipe"
+
+namespace DataprepRecipeUtils
+{
+	bool IsActorValid(const AActor* Actor)
+	{
+		if (Actor == nullptr)
+		{
+			return false;
+		}
+
+		// Don't consider transient actors in non-play worlds
+		// Don't consider the builder brush
+		// Don't consider the WorldSettings actor, even though it is technically editable
+		bool bIsValid = Actor->IsEditable() && !Actor->IsTemplate() && !Actor->HasAnyFlags(RF_Transient) && !FActorEditorUtils::IsABuilderBrush(Actor) && !Actor->IsA(AWorldSettings::StaticClass());
+
+		if (Actor->GetWorld() == GWorld)
+		{
+			// Only consider actors that are allowed to be selected and drawn in editor
+			bIsValid &= Actor->IsListedInSceneOutliner();
+		}
+
+		return bIsValid;
+	}
+}
+
+UDataprepRecipe::UDataprepRecipe()
+{
+	TargetWorld = nullptr;
+}
+
+TArray<TWeakObjectPtr<UObject>> UDataprepRecipe::GetValidAssets(bool bFlushAssets /*= true*/)
+{
+	TArray<TWeakObjectPtr<UObject>> ValidAssets;
+
+#if WITH_EDITOR
+	for (const TWeakObjectPtr<UObject>& AssetPtr : Assets)
+	{
+		if (AssetPtr.Get() != nullptr && !AssetPtr->IsPendingKill())
+		{
+			ValidAssets.Add(AssetPtr);
+		}
+	}
+
+	// Empty data prep recipe's list of assets
+	if (bFlushAssets == true)
+	{
+		Assets.Empty();
+	}
+#endif
+
+	return ValidAssets;
+}
+
+TArray<AActor*> UDataprepRecipe::GetActors() const
+{
+	TArray<AActor*> Result;
+
+	if (TargetWorld == nullptr)
+	{
+		return Result;
+	}
+
+#if WITH_EDITORONLY_DATA
+	const EActorIteratorFlags Flags = EActorIteratorFlags::SkipPendingKill;
+	for (TActorIterator<AActor> It(TargetWorld, AActor::StaticClass(), Flags); It; ++It)
+	{
+		AActor* Actor = *It;
+		if (DataprepRecipeUtils::IsActorValid(Actor))
+		{
+			Result.Add(Actor);
+		}
+	}
+#endif
+
+	return Result;
+}
+
+TArray<UObject*> UDataprepRecipe::GetAssets() const
+{
+	TArray<UObject*> Result;
+
+	for (const TWeakObjectPtr<UObject>& AssetPtr : Assets)
+	{
+		if (AssetPtr.Get() != nullptr && !AssetPtr->IsPendingKill())
+		{
+			Result.Add(AssetPtr.Get());
+		}
+	}
+
+	return Result;
+}
+
+#undef LOCTEXT_NAMESPACE
diff --git a/Engine/Plugins/Enterprise/DataprepEditor/Source/DataprepCore/Public/DataprepAsset.h b/Engine/Plugins/Enterprise/DataprepEditor/Source/DataprepCore/Public/DataprepAsset.h
new file mode 100644
--- /dev/null
+++ b/Engine/Plugins/Enterprise/DataprepEditor/Source/DataprepCore/Public/DataprepAsset.h
@@ -0,0 +1,159 @@
+// Copyright 1998-2019 Epic Games, Inc. All Rights Reserved.
+
+#pragma once
+
+#include "CoreMinimal.h"
+
+#include "DataprepActionAsset.h"
+#include "DataprepAssetInterface.h"
+#include "DataprepAssetProducers.h"
+
+#include "DataprepAsset.generated.h"
+
+class UBlueprint;
+class UDataprepActionAsset;
+class UDataprepParameterizableObject;
+class UDataprepParameterization;
+class UDataprepProducers;
+class UEdGraphNode;
+
+struct FDataprepActionContext;
+struct FDataprepConsumerContext;
+struct FDataprepProducerContext;
+
+/**
+ * A DataprepAsset is an implementation of the DataprepAssetInterface using
+ * a Blueprint as the recipe pipeline. The Blueprint is composed of DataprepAction
+ * nodes linearly connected.
+ */
+UCLASS(Experimental, BlueprintType)
+class DATAPREPCORE_API UDataprepAsset : public UDataprepAssetInterface
+{
+	GENERATED_BODY()
+
+public:
+	UDataprepAsset();
+
+	virtual ~UDataprepAsset() = default;
+
+	// UObject interface
+	virtual void PostLoad() override;
+	virtual bool Rename(const TCHAR* NewName/* =nullptr */, UObject* NewOuter/* =nullptr */, ERenameFlags Flags/* =REN_None */) override;
+	// End of UObject interface
+
+	// UDataprepAssetInterface interface
+	virtual void ExecuteRecipe( const TSharedPtr<FDataprepActionContext>& InActionsContext ) override;
+	virtual bool HasActions() const override { return ActionAssets.Num() > 0; }
+private:
+	virtual TArray<UDataprepActionAsset*> GetCopyOfActions(TMap<UObject*,UObject*>& OutOriginalToCopy) const override;
+	// End of UDataprepAssetInterface interface
+
+public:
+	// Temp code for the nodes development
+	bool CreateBlueprint();
+
+	bool CreateParameterization();
+
+	/** @return pointer on the recipe */
+	const UBlueprint* GetRecipeBP() const
+	{ 
+		return DataprepRecipeBP;
+	}
+
+	UBlueprint* GetRecipeBP()
+	{ 
+		return DataprepRecipeBP;
+	}
+
+	// struct to restrict the access scope
+	struct FDataprepBlueprintChangeNotifier
+	{
+	private:
+		friend class FDataprepEditorUtils;
+
+		static void NotifyDataprepBlueprintChange(UDataprepAsset& DataprepAsset, UObject* ModifiedObject)
+		{
+			// The asset is not complete yet. Skip this change
+			if(!DataprepAsset.HasAnyFlags(RF_NeedLoad | RF_NeedPostLoad | RF_NeedPostLoadSubobjects))
+			{
+				if(UDataprepActionAsset* ActionAsset = Cast<UDataprepActionAsset>(ModifiedObject))
+				{
+					DataprepAsset.UpdateActions();
+				}
+				DataprepAsset.OnBlueprintChanged.Broadcast( ModifiedObject );
+			}
+		}
+	};
+
+	/**
+	* Allow an observer to be notified of an change in the pipeline
+	* return The event that will be broadcasted when a object has receive a modification that might change the result of the pipeline
+	*/
+	DECLARE_EVENT_OneParam(UDataprepAsset, FOnDataprepBlueprintChange, UObject* /*The object that was modified*/)
+	FOnDataprepBlueprintChange& GetOnBlueprintChanged() { return OnBlueprintChanged; }
+	// end of temp code for nodes development
+
+public:
+	// Functions specific to the parametrization of the dataprep asset
+
+	/**
+	 * Event to notify the ui that a dataprep parametrization was modified
+	 * This necessary as the ui for the parameterization is only updated by manual event (the ui don't pull new values each frame)
+	 * Note on the objects param: The parameterized objects that should refresh their ui. If nullptr all widgets that can display some info on the parameterization should be refreshed
+	 */
+	DECLARE_EVENT_OneParam(UDataprepParameterization, FDataprepParameterizationStatusForObjectsChanged, const TSet<UObject*>* /** Objects */)
+	FDataprepParameterizationStatusForObjectsChanged OnParameterizedObjectsChanged;
+
+	virtual UObject* GetParameterizationObject() override;
+
+	void BindObjectPropertyToParameterization(UDataprepParameterizableObject* Object, const TArray<struct FDataprepPropertyLink>& InPropertyChain,const FName& Name);
+
+	bool IsObjectPropertyBinded(UDataprepParameterizableObject* Object, const TArray<struct FDataprepPropertyLink>& InPropertyChain) const;
+
+	FName GetNameOfParameterForObjectProperty(UDataprepParameterizableObject* Object, const TArray<struct FDataprepPropertyLink>& InPropertyChain) const;
+
+	void RemoveObjectPropertyFromParameterization(UDataprepParameterizableObject* Object, const TArray<struct FDataprepPropertyLink>& InPropertyChain);
+
+	void GetExistingParameterNamesForType(UProperty* Property, bool bIsDescribingFullProperty, TSet<FString>& OutValidExistingNames, TSet<FString>& OutInvalidNames) const;
+
+	UDataprepParameterization* GetDataprepParameterization() { return Parameterization; }
+
+protected:
+#if WITH_EDITORONLY_DATA
+	// Temp code for the nodes development
+	/** Temporary: Pointer to data preparation pipeline blueprint used to process input data */
+	UPROPERTY()
+	UBlueprint* DataprepRecipeBP;
+	// end of temp code for nodes development
+
+	/** DEPRECATED: List of producers referenced by the asset */
+	UPROPERTY()
+	TArray< FDataprepAssetProducer > Producers_DEPRECATED;
+
+	/** DEPRECATED: COnsumer referenced by the asset */
+	UPROPERTY()
+	UDataprepContentConsumer* Consumer_DEPRECATED;
+#endif
+
+	// Temp code for the nodes development
+	void OnDataprepBlueprintChanged( class UBlueprint* InBlueprint );
+	// end of temp code for nodes development
+
+	// Temp code for the nodes development
+private:
+	void UpdateActions();
+
+private:
+	UPROPERTY()
+	UEdGraphNode* StartNode;
+
+	UPROPERTY()
+	UDataprepParameterization* Parameterization;
+
+	UPROPERTY()
+	TArray<UDataprepActionAsset*> ActionAssets;
+
+	/** Event broadcasted when object in the pipeline was modified (Only broadcasted on changes that can affect the result of execution) */
+	FOnDataprepBlueprintChange OnBlueprintChanged;
+	// end of temp code for nodes development
+};
\ No newline at end of file
diff --git a/Engine/Plugins/Enterprise/DataprepEditor/Source/DataprepCore/Public/DataprepRecipe.h b/Engine/Plugins/Enterprise/DataprepEditor/Source/DataprepCore/Public/DataprepRecipe.h
new file mode 100644
--- /dev/null
+++ b/Engine/Plugins/Enterprise/DataprepEditor/Source/DataprepCore/Public/DataprepRecipe.h
@@ -0,0 +1,95 @@
+// Copyright 1998-2019 Epic Games, Inc. All Rights Reserved.
+
+#pragma once
+
+#include "CoreMinimal.h"
+#include "Engine/Blueprint.h"
+
+#include "DataprepRecipe.generated.h"
+
+class AActor;
+class UWorld;
+
+UCLASS(Experimental, Blueprintable, BlueprintType, meta = (DisplayName = "Data Preparation Recipe"))
+class DATAPREPCORE_API UDataprepRecipe : public UObject
+{
+	GENERATED_BODY()
+
+public:
+	UDataprepRecipe();
+
+	/** Begin UObject override */
+	virtual bool IsEditorOnly() const override { return true; }
+	/** End UObject override */
+
+	/**
+	 * Reset array of assets attached to recipe
+	 */
+	void ResetAssets()
+	{
+		Assets.Reset();
+	}
+
+	/**
+	 * Setter on array of assets held by this object
+	 */
+	void SetAssets(TArray<TWeakObjectPtr<UObject>> InAssets)
+	{
+		Assets = MoveTemp(InAssets);
+	}
+
+	/**
+	 * Returns array of assets attached to recipe.
+	 * Usually called after the operations of the recipe have been executed
+	 */
+	TArray<TWeakObjectPtr<UObject>> GetValidAssets(bool bFlushAssets = true);
+
+	/**
+	 * Returns world attached to this recipe
+	 */
+	UWorld* GetTargetWorld()
+	{
+		return TargetWorld;
+	}
+
+	/**
+	 * Attaches a world to this recipe
+	 */
+	void SetTargetWorld(UWorld* InTargetWorld)
+	{
+		TargetWorld = InTargetWorld;
+	}
+
+public:
+	/**
+	 * Returns all actors contained in its attached world
+	 */
+	UFUNCTION(BlueprintCallable, Category = "Dataprep | Query", meta = (HideSelfPin = "true"))
+	virtual TArray<AActor*> GetActors() const;
+
+	/**
+	 * Returns all valid assets contained in attached world
+	 */
+	UFUNCTION(BlueprintCallable, Category = "Dataprep | Query", meta = (HideSelfPin = "true"))
+	virtual TArray<UObject*> GetAssets() const;
+
+	/**
+	 * Function used to trigger the execution of the pipeline
+	 * An event node associated to this function must be in the pipeline graph to run it.
+	 */
+	UFUNCTION(BlueprintImplementableEvent)
+	void TriggerPipelineTraversal();
+
+private:
+	/**
+	 * World attached to this recipe.
+	 * All queries on actors will be perform from this world.
+	 * Transient, Can be null.
+	 */
+	UWorld* TargetWorld;
+
+	/**
+	 * All queries on assets will be perform from this array.
+	 */
+	TArray<TWeakObjectPtr<UObject>> Assets;
+};
diff --git a/Engine/Plugins/Enterprise/DataprepEditor/Source/DataprepCore/Public/SelectionSystem/DataprepFloatFilter.h b/Engine/Plugins/Enterprise/DataprepEditor/Source/DataprepCore/Public/SelectionSystem/DataprepFloatFilter.h
new file mode 100644
--- /dev/null
+++ b/Engine/Plugins/Enterprise/DataprepEditor/Source/DataprepCore/Public/SelectionSystem/DataprepFloatFilter.h
@@ -0,0 +1,64 @@
+// Copyright 1998-2019 Epic Games, Inc. All Rights Reserved.
+
+#pragma once
+
+#include "SelectionSystem/DataprepFilter.h"
+
+#include "CoreMinimal.h"
+#include "UObject/Object.h"
+#include "UObject/ObjectMacros.h"
+
+#include "DataprepFloatFilter.generated.h"
+
+class UDataprepFloatFetcher;
+
+UENUM()
+enum class EDataprepFloatMatchType : uint8
+{
+	LessThan,
+	GreatherThan,
+	IsNearlyEqual
+};
+
+UCLASS()
+class DATAPREPCORE_API UDataprepFloatFilter : public UDataprepFilter
+{
+	GENERATED_BODY()
+
+public:
+	bool Filter(float Float) const;
+
+	//~ Begin UDataprepFilter Interface
+	virtual TArray<UObject*> FilterObjects(const TArray<UObject*>& Objects) const override;
+	virtual bool IsThreadSafe() const override { return true; }
+	virtual FText GetFilterCategoryText() const override;
+	virtual TSubclassOf<UDataprepFetcher> GetAcceptedFetcherClass() const override;
+	virtual void SetFetcher(const TSubclassOf<UDataprepFetcher>& FetcherClass) override;
+	virtual UDataprepFetcher* GetFetcher() const override;
+	//~ Begin UDataprepFilter Interface
+
+	EDataprepFloatMatchType GetFloatMatchingCriteria() const;
+	float GetEqualValue() const;
+	float GetTolerance() const;
+
+	void SetFloatMatchingCriteria(EDataprepFloatMatchType FloatMatchingCriteria);
+	void SetEqualValue(float EqualValue);
+	void SetTolerance(float Tolerance);
+
+private:
+	// The source of float selected by the user
+	UPROPERTY()
+	UDataprepFloatFetcher* FloatFetcher;
+
+	// The criteria selected by the user
+	UPROPERTY(EditAnywhere, Category = Filter)
+	EDataprepFloatMatchType FloatMatchingCriteria;
+
+	// The value to use for the equality check
+	UPROPERTY(EditAnywhere, Category = Filter)
+	float EqualValue;
+
+	// The value used for the tolerance when doing a nearly equal
+	UPROPERTY(EditAnywhere, Category = Filter)
+	float Tolerance = KINDA_SMALL_NUMBER;
+};
diff --git a/Engine/Plugins/Enterprise/DataprepEditor/Source/DataprepEditor/Private/AssetTypeActions_Dataprep.cpp b/Engine/Plugins/Enterprise/DataprepEditor/Source/DataprepEditor/Private/AssetTypeActions_Dataprep.cpp
new file mode 100644
--- /dev/null
+++ b/Engine/Plugins/Enterprise/DataprepEditor/Source/DataprepEditor/Private/AssetTypeActions_Dataprep.cpp
@@ -0,0 +1,32 @@
+// Copyright 1998-2019 Epic Games, Inc. All Rights Reserved.
+
+#include "AssetTypeActions_Dataprep.h"
+
+#include "DataprepEditorModule.h"
+#include "DataprepRecipe.h"
+
+/* FAssetTypeActions_Base interface
+*****************************************************************************/
+
+uint32 FAssetTypeActions_Dataprep::GetCategories()
+{
+	return IDataprepEditorModule::DataprepCategoryBit;
+}
+
+
+FText FAssetTypeActions_Dataprep::GetName() const
+{
+	return NSLOCTEXT("AssetTypeActions_Dataprep", "AssetTypeActions_Dataprep_Name", "Datasmith Dataprep");
+}
+
+
+UClass* FAssetTypeActions_Dataprep::GetSupportedClass() const
+{
+	return UDataprepRecipe::StaticClass();
+}
+
+
+FColor FAssetTypeActions_Dataprep::GetTypeColor() const
+{
+	return FColor::White;
+}
diff --git a/Engine/Plugins/Enterprise/DataprepEditor/Source/DataprepEditor/Private/AssetTypeActions_Dataprep.h b/Engine/Plugins/Enterprise/DataprepEditor/Source/DataprepEditor/Private/AssetTypeActions_Dataprep.h
new file mode 100644
--- /dev/null
+++ b/Engine/Plugins/Enterprise/DataprepEditor/Source/DataprepEditor/Private/AssetTypeActions_Dataprep.h
@@ -0,0 +1,18 @@
+// Copyright 1998-2019 Epic Games, Inc. All Rights Reserved.
+
+#pragma once
+
+#include "CoreMinimal.h"
+#include "AssetTypeActions_Base.h"
+
+/** Asset type actions for UDataprepRecipe class */
+class FAssetTypeActions_Dataprep : public FAssetTypeActions_Base
+{
+public:
+	// Begin IAssetTypeActions interface
+	virtual uint32 GetCategories() override;
+	virtual FText GetName() const override;
+	virtual UClass* GetSupportedClass() const override;
+	virtual FColor GetTypeColor() const override;
+	// End IAssetTypeActions interface
+};
diff --git a/Engine/Plugins/Enterprise/DataprepEditor/Source/DataprepEditor/Private/AssetTypeActions_DataprepAsset.cpp b/Engine/Plugins/Enterprise/DataprepEditor/Source/DataprepEditor/Private/AssetTypeActions_DataprepAsset.cpp
new file mode 100644
--- /dev/null
+++ b/Engine/Plugins/Enterprise/DataprepEditor/Source/DataprepEditor/Private/AssetTypeActions_DataprepAsset.cpp
@@ -0,0 +1,61 @@
+// Copyright 1998-2019 Epic Games, Inc. All Rights Reserved.
+
+#include "AssetTypeActions_DataprepAsset.h"
+
+#include "DataprepAsset.h"
+#include "DataprepAssetInstance.h"
+#include "DataprepEditor.h"
+
+FText FAssetTypeActions_DataprepAsset::GetName() const
+{
+	return NSLOCTEXT("AssetTypeActions_DataprepAsset", "Name", "Dataprep Asset");
+}
+
+UClass* FAssetTypeActions_DataprepAsset::GetSupportedClass() const
+{
+	return UDataprepAsset::StaticClass();
+}
+
+void FAssetTypeActions_DataprepAsset::OpenAssetEditor(const TArray<UObject*>& InObjects, TSharedPtr<class IToolkitHost> EditWithinLevelEditor)
+{
+	if (InObjects.Num() == 0)
+	{
+		return;
+	}
+
+	for (UObject* Object : InObjects)
+	{
+		if (UDataprepAsset* DataprepAsset = Cast<UDataprepAsset>(Object))
+		{
+			TSharedRef<FDataprepEditor> NewDataprepEditor(new FDataprepEditor());
+			NewDataprepEditor->InitDataprepEditor( EToolkitMode::Standalone, EditWithinLevelEditor, DataprepAsset, DataprepAsset->GetRecipeBP() );
+		}
+	}
+}
+
+FText FAssetTypeActions_DataprepAssetInstance::GetName() const
+{
+	return NSLOCTEXT("AssetTypeActions_DataprepAssetInstance", "Name", "Dataprep Asset Instance");
+}
+
+UClass* FAssetTypeActions_DataprepAssetInstance::GetSupportedClass() const
+{
+	return UDataprepAssetInstance::StaticClass();
+}
+
+void FAssetTypeActions_DataprepAssetInstance::OpenAssetEditor(const TArray<UObject*>& InObjects, TSharedPtr<class IToolkitHost> EditWithinLevelEditor)
+{
+	if (InObjects.Num() == 0)
+	{
+		return;
+	}
+
+	for (UObject* Object : InObjects)
+	{
+		if (UDataprepAssetInstance* DataprepAssetInstance = Cast<UDataprepAssetInstance>(Object))
+		{
+			TSharedRef<FDataprepEditor> NewDataprepEditor(new FDataprepEditor());
+			NewDataprepEditor->InitDataprepEditor( EToolkitMode::Standalone, EditWithinLevelEditor, DataprepAssetInstance );
+		}
+	}
+}
diff --git a/Engine/Plugins/Enterprise/DataprepEditor/Source/DataprepEditor/Private/AssetTypeActions_DataprepAsset.h b/Engine/Plugins/Enterprise/DataprepEditor/Source/DataprepEditor/Private/AssetTypeActions_DataprepAsset.h
new file mode 100644
--- /dev/null
+++ b/Engine/Plugins/Enterprise/DataprepEditor/Source/DataprepEditor/Private/AssetTypeActions_DataprepAsset.h
@@ -0,0 +1,32 @@
+// Copyright 1998-2019 Epic Games, Inc. All Rights Reserved.
+
+#pragma once
+
+#include "CoreMinimal.h"
+
+#include "AssetTypeActions_DataprepAssetInterface.h"
+
+/** Asset type actions for UDatasmithScene class */
+class FAssetTypeActions_DataprepAsset : public FAssetTypeActions_DataprepAssetInterface
+{
+public:
+	// Begin IAssetTypeActions interface
+	virtual FText GetName() const override;
+	virtual FColor GetTypeColor() const override { return FColor(255, 255, 0); }
+	virtual UClass* GetSupportedClass() const override;
+	virtual void OpenAssetEditor(const TArray<UObject*>& InObjects, TSharedPtr<class IToolkitHost> EditWithinLevelEditor = TSharedPtr<IToolkitHost>()) override;
+	// End IAssetTypeActions interface
+};
+
+/** Asset type actions for UDatasmithScene class */
+class FAssetTypeActions_DataprepAssetInstance : public FAssetTypeActions_DataprepAssetInterface
+{
+public:
+	// Begin IAssetTypeActions interface
+	virtual FText GetName() const override;
+	virtual FColor GetTypeColor() const override { return FColor(255, 0, 255); }
+	virtual UClass* GetSupportedClass() const override;
+	virtual void OpenAssetEditor(const TArray<UObject*>& InObjects, TSharedPtr<class IToolkitHost> EditWithinLevelEditor = TSharedPtr<IToolkitHost>()) override;
+	// End IAssetTypeActions interface
+};
+
diff --git a/Engine/Plugins/Enterprise/DataprepEditor/Source/DataprepEditor/Private/DataprepDetailsTab.cpp b/Engine/Plugins/Enterprise/DataprepEditor/Source/DataprepEditor/Private/DataprepDetailsTab.cpp
new file mode 100644
--- /dev/null
+++ b/Engine/Plugins/Enterprise/DataprepEditor/Source/DataprepEditor/Private/DataprepDetailsTab.cpp
@@ -0,0 +1,127 @@
+// Copyright 1998-2019 Epic Games, Inc. All Rights Reserved.
+
+#include "DataprepEditor.h"
+
+#include "Widgets/DataprepAssetView.h"
+#include "DataprepContentConsumer.h"
+#include "DataprepContentProducer.h"
+#include "DataprepEditorActions.h"
+#include "DataprepEditorModule.h"
+#include "DataprepEditorStyle.h"
+#include "DataprepRecipe.h"
+#include "Widgets/SAssetsPreviewWidget.h"
+
+#include "AssetRegistryModule.h"
+#include "DesktopPlatformModule.h"
+#include "Dialogs/DlgPickPath.h"
+#include "EditorDirectories.h"
+#include "Editor.h"
+#include "EditorStyleSet.h"
+#include "Engine/Blueprint.h"
+#include "Engine/SCS_Node.h"
+#include "Engine/World.h"
+#include "EngineUtils.h"
+#include "Framework/Application/SlateApplication.h"
+#include "Framework/MultiBox/MultiBoxBuilder.h"
+#include "K2Node_AddComponent.h"
+#include "K2Node_CustomEvent.h"
+#include "Kismet2/BlueprintEditorUtils.h"
+#include "Kismet2/CompilerResultsLog.h"
+#include "Kismet2/KismetEditorUtilities.h"
+#include "Materials/MaterialInstance.h"
+#include "Misc/Attribute.h"
+#include "Misc/Guid.h"
+#include "Misc/Paths.h"
+#include "Misc/ScopedSlowTask.h"
+#include "Modules/ModuleManager.h"
+#include "PropertyEditorModule.h"
+#include "SceneOutlinerModule.h"
+#include "ScopedTransaction.h"
+#include "Toolkits/IToolkit.h"
+#include "UnrealEdGlobals.h"
+
+#include "Widgets/Docking/SDockTab.h"
+#include "Widgets/Images/SImage.h"
+#include "Widgets/Input/STextComboBox.h"
+#include "Widgets/Input/SSearchBox.h"
+#include "Widgets/Input/SButton.h"
+#include "Widgets/Input/SComboButton.h"
+#include "Widgets/Layout/SBorder.h"
+#include "Widgets/Layout/SScrollBox.h"
+#include "Widgets/Text/SInlineEditableTextBlock.h"
+#include "Widgets/Text/STextBlock.h"
+#include "Widgets/Views/STableRow.h"
+
+#define LOCTEXT_NAMESPACE "DataprepEditor"
+
+TSharedRef<SDockTab> FDataprepEditor::SpawnTabDetails(const FSpawnTabArgs & Args)
+{
+	check(Args.GetTabId() == DetailsTabId);
+
+	return SAssignNew(DetailsTabPtr, SDockTab)
+		.Label(LOCTEXT("DataprepEditor_DetailsTab_Title", "Details"))
+		[
+			SNullWidget::NullWidget
+		];
+}
+
+void FDataprepEditor::CreateDetailsViews()
+{
+	DetailsView = SNew( SGraphNodeDetailsWidget );
+}
+
+void FDataprepEditor::OnPipelineEditorSelectionChanged(const TSet<UObject*>& SelectedNodes)
+{
+	if (DetailsView.IsValid())
+	{
+		SetDetailsObjects(SelectedNodes, true);
+	}
+}
+
+void FDataprepEditor::SetDetailsObjects(const TSet<UObject*>& Objects, bool bCanEditProperties)
+{
+	if (DetailsView.IsValid())
+	{
+		DetailsView->ShowDetailsObjects(Objects.Array());
+		DetailsView->SetCanEditProperties(bCanEditProperties);
+
+		if ( DetailsTabPtr.IsValid() && DetailsTabPtr.Pin()->GetContent() != DetailsView )
+		{
+			DetailsTabPtr.Pin()->SetContent( DetailsView.ToSharedRef() );
+		}
+	}
+}
+
+void FDataprepEditor::TryInvokingDetailsTab(bool bFlash)
+{
+	if (TabManager->HasTabSpawner(FDataprepEditor::DetailsTabId))
+	{
+		TSharedPtr<SDockTab> DataprepTab = FGlobalTabmanager::Get()->GetMajorTabForTabManager(TabManager.ToSharedRef());
+
+		// We don't want to force this tab into existence when the data prep editor isn't in the foreground and actively
+		// being interacted with.  So we make sure the window it's in is focused and the tab is in the foreground.
+		if (DataprepTab.IsValid() && DataprepTab->IsForeground())
+		{
+			TSharedPtr<SWindow> ParentWindow = DataprepTab->GetParentWindow();
+			if (ParentWindow.IsValid() && ParentWindow->HasFocusedDescendants())
+			{
+				if (!DetailsTabPtr.IsValid())
+				{
+					// Show the details panel if it doesn't exist.
+					TabManager->InvokeTab(FDataprepEditor::DetailsTabId);
+				}
+			}
+		}
+	}
+
+	if (bFlash)
+	{
+		if (DetailsTabPtr.IsValid())
+		{
+			TSharedPtr<SDockTab> DetailsTab = DetailsTabPtr.Pin();
+			DetailsTab->FlashTab();
+		}
+	}
+}
+
+#undef LOCTEXT_NAMESPACE
diff --git a/Engine/Plugins/Enterprise/DataprepEditor/Source/DataprepEditor/Private/DataprepEditor.cpp b/Engine/Plugins/Enterprise/DataprepEditor/Source/DataprepEditor/Private/DataprepEditor.cpp
new file mode 100644
--- /dev/null
+++ b/Engine/Plugins/Enterprise/DataprepEditor/Source/DataprepEditor/Private/DataprepEditor.cpp
@@ -0,0 +1,1274 @@
+// Copyright 1998-2019 Epic Games, Inc. All Rights Reserved.
+
+#include "DataprepEditor.h"
+
+#include "DataprepOperation.h"
+
+#include "DataprepAssetInstance.h"
+#include "DataprepContentConsumer.h"
+#include "DataprepContentProducer.h"
+#include "DataprepEditorActions.h"
+#include "DataprepEditorModule.h"
+#include "DataprepEditorStyle.h"
+#include "DataprepRecipe.h"
+#include "DataprepActionAsset.h"
+#include "DataprepCoreUtils.h"
+#include "DataprepEditorLogCategory.h"
+#include "SchemaActions/DataprepOperationMenuActionCollector.h"
+#include "Widgets/DataprepAssetView.h"
+#include "Widgets/SAssetsPreviewWidget.h"
+#include "Widgets/SDataprepEditorViewport.h"
+#include "Widgets/SDataprepPalette.h"
+
+#include "ActorEditorUtils.h"
+#include "ActorTreeItem.h"
+#include "AssetDeleteModel.h"
+#include "AssetRegistryModule.h"
+#include "StatsViewerModule.h"
+#include "BlueprintNodeSpawner.h"
+#include "ComponentTreeItem.h"
+#include "DesktopPlatformModule.h"
+#include "Dialogs/Dialogs.h"
+#include "Dialogs/DlgPickPath.h"
+#include "Editor.h"
+#include "EditorDirectories.h"
+#include "EditorStyleSet.h"
+#include "Engine/World.h"
+#include "EngineUtils.h"
+#include "Framework/Application/SlateApplication.h"
+#include "Framework/MultiBox/MultiBoxBuilder.h"
+#include "GenericPlatform/GenericPlatformTime.h"
+#include "HAL/FileManager.h"
+#include "ICustomSceneOutliner.h"
+#include "Kismet2/BlueprintEditorUtils.h"
+#include "Kismet2/KismetEditorUtilities.h"
+#include "Materials/MaterialInstance.h"
+#include "Misc/Attribute.h"
+#include "Misc/Guid.h"
+#include "Misc/Paths.h"
+#include "Misc/ScopedSlowTask.h"
+#include "Modules/ModuleManager.h"
+#include "ObjectTools.h"
+#include "ScopedTransaction.h"
+#include "Templates/UnrealTemplate.h"
+#include "Toolkits/IToolkit.h"
+#include "UObject/Object.h"
+#include "UnrealEdGlobals.h"
+#include "Widgets/Docking/SDockTab.h"
+#include "Widgets/Input/STextComboBox.h"
+#include "Widgets/Layout/SBorder.h"
+
+// Temp code for the nodes development
+#include "Engine/Blueprint.h"
+#include "Kismet2/BlueprintEditorUtils.h"
+#include "K2Node_CustomEvent.h"
+#include "DataprepRecipe.h"
+#include "Engine/BlueprintGeneratedClass.h"
+#include "Kismet2/KismetEditorUtilities.h"
+#include "PropertyEditorModule.h"
+#include "IDetailsView.h"
+#include "DataprepAsset.h"
+#include "IStructureDetailsView.h"
+
+const FName FDataprepEditor::PipelineGraphTabId(TEXT("DataprepEditor_Pipeline_Graph"));
+// end of temp code for nodes development
+
+#define LOCTEXT_NAMESPACE "DataprepEditor"
+
+extern const FName DataprepEditorAppIdentifier;
+
+const FName FDataprepEditor::ScenePreviewTabId(TEXT("DataprepEditor_ScenePreview"));
+const FName FDataprepEditor::AssetPreviewTabId(TEXT("DataprepEditor_AssetPreview"));
+const FName FDataprepEditor::PaletteTabId(TEXT("DataprepEditor_Palette"));
+const FName FDataprepEditor::DetailsTabId(TEXT("DataprepEditor_Details"));
+const FName FDataprepEditor::DataprepAssetTabId(TEXT("DataprepEditor_Dataprep"));
+const FName FDataprepEditor::SceneViewportTabId(TEXT("DataprepEditor_SceneViewport"));
+const FName FDataprepEditor::DataprepStatisticsTabId(TEXT("DataprepEditor_Statistics"));
+
+static bool bLogTiming = true;
+
+class FTimeLogger
+{
+public:
+	FTimeLogger(const FString& InText)
+		: StartTime( FPlatformTime::Cycles64() )
+		, Text( InText )
+	{
+		if( bLogTiming )
+		{
+			UE_LOG( LogDataprepEditor, Log, TEXT("%s ..."), *Text );
+		}
+	}
+
+	~FTimeLogger()
+	{
+		if( bLogTiming )
+		{
+			// Log time spent to import incoming file in minutes and seconds
+			double ElapsedSeconds = FPlatformTime::ToSeconds64(FPlatformTime::Cycles64() - StartTime);
+
+			int ElapsedMin = int(ElapsedSeconds / 60.0);
+			ElapsedSeconds -= 60.0 * (double)ElapsedMin;
+			UE_LOG( LogDataprepEditor, Log, TEXT("%s took [%d min %.3f s]"), *Text, ElapsedMin, ElapsedSeconds );
+		}
+	}
+
+private:
+	uint64 StartTime;
+	FString Text;
+};
+
+namespace DataprepEditorUtil
+{
+	void DeleteActor(AActor* Actor, UWorld* World)
+	{
+		if (Actor == nullptr || World != Actor->GetWorld())
+		{
+			return;
+		}
+
+		TArray<AActor*> Children;
+		Actor->GetAttachedActors(Children);
+
+		for (AActor* ChildActor : Children)
+		{
+			DeleteActor(ChildActor, World);
+		}
+
+		World->DestroyActor(Actor, false, true);
+	}
+
+	void DeleteTemporaryPackage( const FString& PathToDelete );
+}
+
+FDataprepEditor::FDataprepEditor()
+	: bWorldBuilt(false)
+	, bIsFirstRun(false)
+	, bPipelineChanged(false)
+	, bIsDataprepInstance(false)
+	, bIsActionMenuContextSensitive(true)
+	, bSaveIntermediateBuildProducts(false)
+	, PreviewWorld(nullptr)
+	, bIgnoreCloseRequest(false)
+{
+	FName UniqueWorldName = MakeUniqueObjectName(GetTransientPackage(), UWorld::StaticClass(), FName( *(LOCTEXT("PreviewWorld", "Preview").ToString()) ));
+	PreviewWorld = TStrongObjectPtr<UWorld>(NewObject< UWorld >(GetTransientPackage(), UniqueWorldName));
+	PreviewWorld->WorldType = EWorldType::EditorPreview;
+
+	FWorldContext& WorldContext = GEngine->CreateNewWorldContext(PreviewWorld->WorldType);
+	WorldContext.SetCurrentWorld(PreviewWorld.Get());
+
+	PreviewWorld->InitializeNewWorld(UWorld::InitializationValues()
+		.AllowAudioPlayback(false)
+		.CreatePhysicsScene(false)
+		.RequiresHitProxies(false)
+		.CreateNavigation(false)
+		.CreateAISystem(false)
+		.ShouldSimulatePhysics(false)
+		.SetTransactional(false));
+
+	for ( ULevel* Level : PreviewWorld->GetLevels() )
+	{
+		for ( AActor* Actor : Level->Actors )
+		{
+			DefaultActorsInPreviewWorld.Add( Actor );
+		}
+	}
+}
+
+FDataprepEditor::~FDataprepEditor()
+{
+	if( DataprepAssetInterfacePtr.IsValid() )
+	{
+		DataprepAssetInterfacePtr->GetOnChanged().RemoveAll( this );
+	}
+
+	if ( PreviewWorld )
+	{
+		GEngine->DestroyWorldContext( PreviewWorld.Get() );
+		PreviewWorld->DestroyWorld( true );
+		PreviewWorld.Reset();
+	}
+
+	FCoreUObjectDelegates::OnObjectPropertyChanged.RemoveAll(this);
+
+	auto DeleteDirectory = [&](const FString& DirectoryToDelete)
+	{
+		const FString AbsolutePath = FPaths::ConvertRelativePathToFull( DirectoryToDelete );
+		IFileManager::Get().DeleteDirectory( *AbsolutePath, false, true );
+	};
+
+	// Clean up temporary directories and data created for this session
+	{
+
+		DeleteDirectory( TempDir );
+
+		FString PackagePathToDeleteOnDisk;
+		if (FPackageName::TryConvertLongPackageNameToFilename( GetTransientContentFolder(), PackagePathToDeleteOnDisk))
+		{
+			DeleteDirectory( PackagePathToDeleteOnDisk );
+		}
+
+	}
+
+	// Clean up temporary directories associated to process if no session of Dataprep editor is running
+	{
+		auto IsDirectoryEmpty = [&](const TCHAR* Directory) -> bool
+		{
+			bool bDirectoryIsEmpty = true;
+			IFileManager::Get().IterateDirectory(Directory, [&](const TCHAR* FilenameOrDirectory, bool bIsDirectory) -> bool
+			{
+				bDirectoryIsEmpty = false;
+				return false;
+			});
+
+			return bDirectoryIsEmpty;
+		};
+
+		FString RootTempDir = FPaths::Combine( GetRootTemporaryDir(), FString::FromInt( FPlatformProcess::GetCurrentProcessId() ) );
+		if(IsDirectoryEmpty( *RootTempDir ))
+		{
+			DeleteDirectory( RootTempDir );
+		}
+
+		const FString PackagePathToDelete = FPaths::Combine( GetRootPackagePath(), FString::FromInt( FPlatformProcess::GetCurrentProcessId() ) );
+		FString PackagePathToDeleteOnDisk;
+		if (FPackageName::TryConvertLongPackageNameToFilename(PackagePathToDelete, PackagePathToDeleteOnDisk))
+		{
+			if(IsDirectoryEmpty( *PackagePathToDeleteOnDisk ))
+			{
+				DeleteDirectory( PackagePathToDeleteOnDisk );
+			}
+		}
+	}
+}
+
+FName FDataprepEditor::GetToolkitFName() const
+{
+	return FName("DataprepEditor");
+}
+
+FText FDataprepEditor::GetBaseToolkitName() const
+{
+	return LOCTEXT("AppLabel", "Dataprep Editor");
+}
+
+FString FDataprepEditor::GetWorldCentricTabPrefix() const
+{
+	return LOCTEXT("WorldCentricTabPrefix", "Dataprep").ToString();
+}
+
+FLinearColor FDataprepEditor::GetWorldCentricTabColorScale() const
+{
+	return FLinearColor(0.3f, 0.2f, 0.5f, 0.5f);
+}
+
+void FDataprepEditor::RegisterTabSpawners(const TSharedRef<class FTabManager>& InTabManager)
+{
+	WorkspaceMenuCategory = InTabManager->AddLocalWorkspaceMenuCategory(LOCTEXT("WorkspaceMenu_DataprepEditor", "Data Preparation Editor"));
+
+	FAssetEditorToolkit::RegisterTabSpawners(InTabManager);
+
+	auto WorkspaceMenuCategoryRef = WorkspaceMenuCategory.ToSharedRef();
+
+	InTabManager->RegisterTabSpawner(ScenePreviewTabId, FOnSpawnTab::CreateSP(this, &FDataprepEditor::SpawnTabScenePreview))
+		.SetDisplayName(LOCTEXT("ScenePreviewTab", "Scene Preview"))
+		.SetGroup(WorkspaceMenuCategoryRef)
+		.SetIcon(FSlateIcon(FDataprepEditorStyle::GetStyleSetName(), "DataprepEditor.Tabs.ScenePreview"));
+
+	InTabManager->RegisterTabSpawner(AssetPreviewTabId, FOnSpawnTab::CreateSP(this, &FDataprepEditor::SpawnTabAssetPreview))
+		.SetDisplayName(LOCTEXT("AssetPreviewTab", "Asset Preview"))
+		.SetGroup(WorkspaceMenuCategoryRef)
+		.SetIcon(FSlateIcon(FDataprepEditorStyle::GetStyleSetName(), "DataprepEditor.Tabs.AssetPreview"));
+
+	InTabManager->RegisterTabSpawner(SceneViewportTabId, FOnSpawnTab::CreateSP(this, &FDataprepEditor::SpawnTabSceneViewport))
+		.SetDisplayName(LOCTEXT("SceneViewportTab", "Scene Viewport"))
+		.SetGroup(WorkspaceMenuCategoryRef)
+		.SetIcon(FSlateIcon(FDataprepEditorStyle::GetStyleSetName(), "DataprepEditor.Tabs.SceneViewport"));
+
+	InTabManager->RegisterTabSpawner(DetailsTabId, FOnSpawnTab::CreateSP(this, &FDataprepEditor::SpawnTabDetails))
+		.SetDisplayName(LOCTEXT("DetailsTab", "Details"))
+		.SetGroup(WorkspaceMenuCategoryRef)
+		.SetIcon(FSlateIcon(FEditorStyle::GetStyleSetName(), "LevelEditor.Tabs.Details"));
+
+	InTabManager->RegisterTabSpawner(DataprepAssetTabId, FOnSpawnTab::CreateSP(this, &FDataprepEditor::SpawnTabDataprep))
+		.SetDisplayName(LOCTEXT("DataprepAssetTab", "Dataprep"))
+		.SetGroup(WorkspaceMenuCategoryRef)
+		.SetIcon(FSlateIcon(FEditorStyle::GetStyleSetName(), "LevelEditor.Tabs.Details"));
+
+	InTabManager->RegisterTabSpawner(DataprepStatisticsTabId, FOnSpawnTab::CreateSP(this, &FDataprepEditor::SpawnTabStatistics))
+		.SetDisplayName(LOCTEXT("StatisticsTab", "Statistics"))
+		.SetGroup(WorkspaceMenuCategoryRef)
+		.SetIcon(FSlateIcon(FEditorStyle::GetStyleSetName(), "LevelEditor.Tabs.StatsViewer"));
+
+	// Do not register tabs which are not pertinent to Dataprep instance
+	if(!bIsDataprepInstance)
+	{
+		InTabManager->RegisterTabSpawner(PaletteTabId, FOnSpawnTab::CreateSP(this, &FDataprepEditor::SpawnTabPalette))
+			.SetDisplayName(LOCTEXT("PaletteTab", "Palette"))
+			.SetGroup(WorkspaceMenuCategoryRef)
+			.SetIcon( FSlateIcon(FEditorStyle::GetStyleSetName(), "Kismet.Tabs.Palette"));
+
+		// Temp code for the nodes development
+		InTabManager->RegisterTabSpawner(PipelineGraphTabId, FOnSpawnTab::CreateSP(this, &FDataprepEditor::SpawnTabPipelineGraph))
+			.SetDisplayName(LOCTEXT("PipelineGraphTab", "Pipeline Graph"))
+			.SetGroup(WorkspaceMenuCategoryRef)
+			.SetIcon(FSlateIcon(FEditorStyle::GetStyleSetName(), "GraphEditor.EventGraph_16x"));
+		// end of temp code for nodes development
+	}
+}
+
+void FDataprepEditor::UnregisterTabSpawners(const TSharedRef<class FTabManager>& InTabManager)
+{
+	FAssetEditorToolkit::UnregisterTabSpawners(InTabManager);
+	InTabManager->UnregisterAllTabSpawners();
+}
+
+void FDataprepEditor::CleanUpTemporaryDirectories()
+{
+	const int32 CurrentProcessID = FPlatformProcess::GetCurrentProcessId();
+
+	TSet<FString> TempDirectories;
+	IFileManager::Get().IterateDirectory( *GetRootTemporaryDir(), [&](const TCHAR* FilenameOrDirectory, bool bIsDirectory) -> bool
+	{
+		if (bIsDirectory)
+		{
+			FString DirectoryName = FPaths::GetBaseFilename( FilenameOrDirectory );
+			if(FCString::IsNumeric( *DirectoryName ))
+			{
+
+				int32 ProcessID = FCString::Atoi( *DirectoryName );
+				if(ProcessID != CurrentProcessID)
+				{
+					FProcHandle ProcHandle = FPlatformProcess::OpenProcess( ProcessID );
+
+					// Delete directories if process is not valid
+					bool bDeleteDirectories = !ProcHandle.IsValid();
+
+					// Process is valid, check if application associated with process id is UE4 editor
+					if(!bDeleteDirectories)
+					{
+						const FString ApplicationName = FPlatformProcess::GetApplicationName( ProcessID );
+						bDeleteDirectories = !ApplicationName.StartsWith(TEXT("UE4Editor"));
+					}
+
+					if(bDeleteDirectories)
+					{
+						FString PackagePathToDelete = FPaths::Combine( GetRootPackagePath(), DirectoryName );
+						FString PackagePathToDeleteOnDisk;
+						if (FPackageName::TryConvertLongPackageNameToFilename(PackagePathToDelete, PackagePathToDeleteOnDisk))
+						{
+							TempDirectories.Add( MoveTemp(PackagePathToDeleteOnDisk) );
+						}
+
+						TempDirectories.Emplace( FilenameOrDirectory );
+					}
+				}
+			}
+		}
+		return true;
+	});
+
+	for(FString& TempDirectory : TempDirectories)
+	{
+		FString AbsolutePath = FPaths::ConvertRelativePathToFull( TempDirectory );
+		IFileManager::Get().DeleteDirectory( *AbsolutePath, false, true );
+	}
+}
+
+const FString& FDataprepEditor::GetRootTemporaryDir()
+{
+	static FString RootTemporaryDir = FPaths::Combine(FPaths::ProjectIntermediateDir(), TEXT("DataprepTemp") );
+	return RootTemporaryDir;
+}
+
+const FString& FDataprepEditor::GetRootPackagePath()
+{
+	static FString RootPackagePath( TEXT("/Engine/DataprepEditor/Transient") );
+	return RootPackagePath;
+}
+
+void FDataprepEditor::InitDataprepEditor(const EToolkitMode::Type Mode, const TSharedPtr<IToolkitHost>& InitToolkitHost, UDataprepAssetInterface* InDataprepAssetInterface, UObject* Blueprint)
+{
+	DataprepAssetInterfacePtr = TWeakObjectPtr<UDataprepAssetInterface>(InDataprepAssetInterface);
+	check( DataprepAssetInterfacePtr.IsValid() );
+
+	bIsDataprepInstance = Cast<UDataprepAssetInstance>(InDataprepAssetInterface) != nullptr;
+
+	DataprepAssetInterfacePtr->GetOnChanged().AddSP( this, &FDataprepEditor::OnDataprepAssetChanged );
+
+	// Assign unique session identifier
+	SessionID = FGuid::NewGuid().ToString();
+
+	// Initialize Actions' context
+	DataprepActionAsset::FCanExecuteNextStepFunc CanExecuteNextStepFunc = [this](UDataprepActionAsset* ActionAsset, UDataprepOperation* Operation, UDataprepFilter* Filter) -> bool
+	{
+		return this->OnCanExecuteNextStep(ActionAsset, Operation, Filter);
+	};
+
+	DataprepActionAsset::FActionsContextChangedFunc ActionsContextChangedFunc = [this](const UDataprepActionAsset* ActionAsset, bool bWorldChanged, bool bAssetsChanged, const TArray< TWeakObjectPtr<UObject> >& NewAssets)
+	{
+		this->OnActionsContextChanged(ActionAsset, bWorldChanged, bAssetsChanged, NewAssets);
+	};
+
+	ActionsContext = MakeShareable( new FDataprepActionContext() );
+
+	ActionsContext->SetTransientContentFolder( GetTransientContentFolder() / DataprepAssetInterfacePtr->GetName() / TEXT("Pipeline") )
+		.SetLogger( TSharedPtr<IDataprepLogger>( new FDataprepCoreUtils::FDataprepLogger ) )
+		.SetCanExecuteNextStep( CanExecuteNextStepFunc )
+		.SetActionsContextChanged( ActionsContextChangedFunc );
+
+	// Create temporary directory to store transient data
+	CleanUpTemporaryDirectories();
+	TempDir = FPaths::Combine( GetRootTemporaryDir(), FString::FromInt( FPlatformProcess::GetCurrentProcessId() ), SessionID);
+	IFileManager::Get().MakeDirectory(*TempDir);
+
+	// Temp code for the nodes development
+	if(Blueprint != nullptr)
+	{
+		DataprepRecipeBPPtr = Cast<UBlueprint>(Blueprint);
+		check( DataprepRecipeBPPtr.IsValid() );
+
+		// Necessary step to regenerate blueprint generated class
+		// Note that this compilation will always succeed as Dataprep node does not have real body
+		{
+			FKismetEditorUtilities::CompileBlueprint( DataprepRecipeBPPtr.Get(), EBlueprintCompileOptions::None, nullptr );
+		}
+	}
+	// End temp code for the nodes development
+
+	GEditor->RegisterForUndo(this);
+
+	// Register our commands. This will only register them if not previously registered
+	FDataprepEditorCommands::Register();
+
+	BindCommands();
+
+	CreateTabs();
+
+	
+	const TSharedRef<FTabManager::FLayout> Layout = DataprepRecipeBPPtr.IsValid() ? CreateDataprepLayout() : CreateDataprepInstanceLayout();
+
+	const bool bCreateDefaultStandaloneMenu = true;
+	const bool bCreateDefaultToolbar = true;
+	FAssetEditorToolkit::InitAssetEditor( Mode, InitToolkitHost, DataprepEditorAppIdentifier, Layout, bCreateDefaultToolbar, bCreateDefaultStandaloneMenu, InDataprepAssetInterface );
+
+	ExtendMenu();
+	ExtendToolBar();
+	RegenerateMenusAndToolbars();
+
+#ifdef DATAPREP_EDITOR_VERBOSE
+	LogDataprepEditor.SetVerbosity( ELogVerbosity::Verbose );
+#endif
+}
+
+void FDataprepEditor::BindCommands()
+{
+	const FDataprepEditorCommands& Commands = FDataprepEditorCommands::Get();
+
+	const TSharedRef<FUICommandList>& UICommandList = GetToolkitCommands();
+
+	//UICommandList->MapAction(FGenericCommands::Get().Delete,
+	//	FExecuteAction::CreateSP(this, &FDataprepEditor::DeleteSelected),
+	//	FCanExecuteAction::CreateSP(this, &FDataprepEditor::CanDeleteSelected));
+
+	//UICommandList->MapAction(FGenericCommands::Get().Undo,
+	//	FExecuteAction::CreateSP(this, &FDataprepEditor::UndoAction));
+
+	//UICommandList->MapAction(FGenericCommands::Get().Redo,
+	//	FExecuteAction::CreateSP(this, &FDataprepEditor::RedoAction));
+
+	// Temp code for the nodes development
+	UICommandList->MapAction(
+		Commands.CompileGraph,
+		FExecuteAction::CreateSP(this, &FDataprepEditor::OnCompile));
+	// end of temp code for nodes development
+
+	UICommandList->MapAction(
+		Commands.SaveScene,
+		FExecuteAction::CreateSP(this, &FDataprepEditor::OnSaveScene));
+
+	UICommandList->MapAction(
+		Commands.BuildWorld,
+		FExecuteAction::CreateSP( this, &FDataprepEditor::OnBuildWorld ),
+		FCanExecuteAction::CreateSP( this, &FDataprepEditor::CanBuildWorld ) );
+
+	UICommandList->MapAction(
+		Commands.ExecutePipeline,
+		FExecuteAction::CreateSP(this, &FDataprepEditor::OnExecutePipeline),
+		FCanExecuteAction::CreateSP(this, &FDataprepEditor::CanExecutePipeline));
+
+	UICommandList->MapAction(
+		Commands.CommitWorld,
+		FExecuteAction::CreateSP(this, &FDataprepEditor::OnCommitWorld),
+		FCanExecuteAction::CreateSP(this, &FDataprepEditor::CanCommitWorld));
+}
+
+void FDataprepEditor::OnSaveScene()
+{
+	TRACE_CPUPROFILER_EVENT_SCOPE(FDataprepEditor::OnSaveScene);
+
+}
+
+void FDataprepEditor::OnBuildWorld()
+{
+	UDataprepAssetInterface* DataprepAssetInterface = DataprepAssetInterfacePtr.Get();
+	if (!ensureAlways(DataprepAssetInterface))
+	{
+		return;
+	}
+
+	if (!ensureAlways(PreviewWorld.IsValid()))
+	{
+		return;
+	}
+
+	TRACE_CPUPROFILER_EVENT_SCOPE(FDataprepEditor::OnBuildWorld);
+
+	if (DataprepAssetInterface->GetProducers()->GetProducersCount() == 0)
+	{
+		ResetBuildWorld();
+		return;
+	}
+
+	CleanPreviewWorld();
+
+	UPackage* TransientPackage = NewObject< UPackage >( nullptr, *GetTransientContentFolder(), RF_Transient );
+	TransientPackage->FullyLoad();
+
+	TSharedPtr< FDataprepCoreUtils::FDataprepFeedbackContext > FeedbackContext( new FDataprepCoreUtils::FDataprepFeedbackContext );
+
+	TSharedPtr< IDataprepProgressReporter > ProgressReporter( new FDataprepCoreUtils::FDataprepProgressUIReporter( FeedbackContext.ToSharedRef() ) );
+	{
+		FTimeLogger TimeLogger( TEXT("Import") );
+
+		FDataprepProducerContext Context;
+		Context.SetWorld( PreviewWorld.Get() )
+			.SetRootPackage( TransientPackage )
+			.SetLogger( TSharedPtr< IDataprepLogger >( new FDataprepCoreUtils::FDataprepLogger ) )
+			.SetProgressReporter( ProgressReporter );
+
+		Assets = DataprepAssetInterface->GetProducers()->Produce( Context );
+	}
+
+	if ( ProgressReporter->IsWorkCancelled() )
+	{
+		// Flush the work that's already been done
+		ResetBuildWorld();
+		return;
+	}
+
+	CachedAssets.Reset();
+	CachedAssets.Append( Assets );
+
+	TakeSnapshot();
+
+	UpdatePreviewPanels();
+	SceneViewportView->FocusViewportOnScene();
+
+	bWorldBuilt = true;
+	bIsFirstRun = true;
+}
+
+void FDataprepEditor::OnDataprepAssetChanged( FDataprepAssetChangeType ChangeType )
+{
+	switch(ChangeType)
+	{
+		case FDataprepAssetChangeType::RecipeModified:
+		{
+			bPipelineChanged = true;
+		}
+		break;
+
+		case FDataprepAssetChangeType::ProducerAdded:
+		case FDataprepAssetChangeType::ProducerRemoved:
+		case FDataprepAssetChangeType::ProducerModified:
+		{
+			// Just reset world for the time being
+			ResetBuildWorld();
+		}
+		break;
+
+		default:
+			break;
+
+	}
+}
+
+void FDataprepEditor::ResetBuildWorld()
+{
+	TRACE_CPUPROFILER_EVENT_SCOPE(FDataprepEditor::ResetBuildWorld);
+
+	bWorldBuilt = false;
+	CleanPreviewWorld();
+	UpdatePreviewPanels();
+	DataprepEditorUtil::DeleteTemporaryPackage( GetTransientContentFolder() );
+}
+
+void FDataprepEditor::CleanPreviewWorld()
+{
+	TRACE_CPUPROFILER_EVENT_SCOPE(FDataprepEditor::CleanPreviewWorld);
+
+	FTimeLogger TimeLogger( TEXT("CleanPreviewWorld") );
+
+	// Destroy all actors in preview world
+	for (ULevel* Level : PreviewWorld->GetLevels())
+	{
+		TArray<AActor*> LevelActors( Level->Actors );
+
+		for( AActor* Actor : LevelActors )
+		{
+			if (Actor && !Actor->IsPendingKill() && !DefaultActorsInPreviewWorld.Contains(Actor))
+			{
+				PreviewWorld->EditorDestroyActor(Actor, true);
+
+				// Since deletion can be delayed, rename to avoid future name collision
+				// Call UObject::Rename directly on actor to avoid AActor::Rename which unnecessarily sunregister and re-register components
+				Actor->UObject::Rename(nullptr, GetTransientPackage(), REN_DontCreateRedirectors | REN_ForceNoResetLoaders);
+			}
+		}
+	}
+
+	SceneViewportView->ClearScene();
+
+	// Delete assets which are still in the transient content folder
+	const FString TransientContentFolder( GetTransientContentFolder() );
+	TArray<UObject*> ObjectsToDelete;
+	for( TWeakObjectPtr<UObject>& Asset : CachedAssets )
+	{
+		if( UObject* ObjectToDelete = Asset.Get() )
+		{
+			FString PackagePath = ObjectToDelete->GetOutermost()->GetName();
+			if( PackagePath.StartsWith( TransientContentFolder ) )
+			{
+				FDataprepCoreUtils::MoveToTransientPackage( ObjectToDelete );
+				ObjectsToDelete.Add( ObjectToDelete );
+			}
+		}
+	}
+
+	// Disable warnings from LogStaticMesh because FDataprepCoreUtils::PurgeObjects is pretty verbose on harmless warnings
+	ELogVerbosity::Type PrevLogStaticMeshVerbosity = LogStaticMesh.GetVerbosity();
+	LogStaticMesh.SetVerbosity( ELogVerbosity::Error );
+
+	FDataprepCoreUtils::PurgeObjects( MoveTemp( ObjectsToDelete ) );
+
+	// Restore LogStaticMesh verbosity
+	LogStaticMesh.SetVerbosity( PrevLogStaticMeshVerbosity );
+
+	CachedAssets.Reset();
+	Assets.Reset();
+
+	PreviewWorld->CleanupActors();
+}
+
+void FDataprepEditor::OnExecutePipeline()
+{
+	if( DataprepAssetInterfacePtr->GetConsumer() == nullptr )
+	{
+		return;
+	}
+
+	TRACE_CPUPROFILER_EVENT_SCOPE(FDataprepEditor::OnExecutePipeline);
+
+	if(!bIsFirstRun)
+	{
+		RestoreFromSnapshot();
+	}
+
+	// Remove any link between assets referenced in the preview world and the viewport
+	SceneViewportView->ClearScene();
+
+	// Trigger execution of data preparation operations on world attached to recipe
+	{
+		FTimeLogger TimeLogger( TEXT("ExecutePipeline") );
+
+		// Some operation can indirectly call FAssetEditorManager::CloseAllAssetEditors (eg. Remove Asset)
+		// Editors can individually refuse this request: we ignore it during the pipeline traversal.
+		TGuardValue<bool> IgnoreCloseRequestGuard(bIgnoreCloseRequest, true);
+
+		TSharedPtr< FDataprepCoreUtils::FDataprepFeedbackContext > FeedbackContext(new FDataprepCoreUtils::FDataprepFeedbackContext);
+		ActionsContext->SetProgressReporter( TSharedPtr< IDataprepProgressReporter >( new FDataprepCoreUtils::FDataprepProgressUIReporter( FeedbackContext.ToSharedRef() ) ) );
+		ActionsContext->SetWorld( PreviewWorld.Get() ).SetAssets( Assets );
+
+		DataprepAssetInterfacePtr->ExecuteRecipe( ActionsContext );
+
+		// Update list of assets with latest ones
+		Assets = ActionsContext->Assets.Array();
+	}
+
+	if ( ActionsContext->ProgressReporterPtr->IsWorkCancelled() )
+	{
+		RestoreFromSnapshot();
+	}
+
+	// Redraw 3D viewport
+	SceneViewportView->UpdateScene();
+
+	// Add assets which may have been created by actions
+	for( TWeakObjectPtr<UObject>& Asset : Assets )
+	{
+		if( Asset.IsValid() )
+		{
+			CachedAssets.Add( Asset );
+		}
+	}
+
+	// Indicate pipeline has been executed at least once
+	bIsFirstRun = false;
+	// Reset tracking of pipeline changes between execution
+	bPipelineChanged = false;
+}
+
+void FDataprepEditor::OnCommitWorld()
+{
+	TRACE_CPUPROFILER_EVENT_SCOPE(FDataprepEditor::OnCommitWorld);
+
+	FTimeLogger TimeLogger( TEXT("Commit") );
+
+	// Pipeline has not been executed, validate with user this is intentional
+	if( bIsFirstRun && DataprepAssetInterfacePtr->HasActions())
+	{
+		const FText Title( LOCTEXT( "DataprepEditor_ProceedWithCommit", "Proceed with commit" ) );
+		const FText Message( LOCTEXT( "DataprepEditor_ConfirmCommitPipelineNotExecuted", "The action pipeline has not been executed.\nDo you want to proceeed with the commit anyway?" ) );
+
+		if( OpenMsgDlgInt( EAppMsgType::YesNo, Message, Title ) == EAppReturnType::No )
+		{
+			return;
+		}
+	}
+	// Pipeline has changed without being executed, validate with user this is intentional
+	else if( !bIsFirstRun && bPipelineChanged )
+	{
+		const FText Title( LOCTEXT( "DataprepEditor_ProceedWithCommit", "Proceed with commit" ) );
+		const FText Message( LOCTEXT( "DataprepEditor_ConfirmCommitPipelineChanged", "The action pipeline has changed since last execution.\nDo you want to proceeed with the commit anyway?" ) );
+
+		if( OpenMsgDlgInt( EAppMsgType::YesNo, Message, Title ) == EAppReturnType::No )
+		{
+			return;
+		}
+	}
+
+	// Remove references to assets in 3D viewport before commit
+	SceneViewportView->ClearScene();
+
+	// Finalize assets
+	TArray<TWeakObjectPtr<UObject>> ValidAssets( MoveTemp(Assets) );
+
+	FDataprepConsumerContext Context;
+	Context.SetWorld( PreviewWorld.Get() )
+		.SetAssets( ValidAssets )
+		.SetTransientContentFolder( GetTransientContentFolder() )
+		.SetLogger( TSharedPtr<IDataprepLogger>( new FDataprepCoreUtils::FDataprepLogger ) )
+		.SetProgressReporter( TSharedPtr< IDataprepProgressReporter >( new FDataprepCoreUtils::FDataprepProgressUIReporter() ) );
+
+	if( !DataprepAssetInterfacePtr->RunConsumer( Context ) )
+	{
+		UE_LOG( LogDataprepEditor, Error, TEXT("Consumer failed...") );
+	}
+
+	ResetBuildWorld();
+}
+
+void FDataprepEditor::ExtendMenu()
+{
+	IDataprepEditorModule* DataprepEditorModule = &FModuleManager::LoadModuleChecked<IDataprepEditorModule>("DataprepEditor");
+	AddMenuExtender(DataprepEditorModule->GetMenuExtensibilityManager()->GetAllExtenders(GetToolkitCommands(), GetEditingObjects()));
+}
+
+void FDataprepEditor::ExtendToolBar()
+{
+	struct Local
+	{
+		static void FillToolbar(FToolBarBuilder& ToolbarBuilder, FDataprepEditor* ThisEditor)
+		{
+			ToolbarBuilder.BeginSection("Run");
+			{
+				ToolbarBuilder.AddToolBarButton(FDataprepEditorCommands::Get().BuildWorld);
+				ToolbarBuilder.AddToolBarButton(FDataprepEditorCommands::Get().ExecutePipeline);
+				ToolbarBuilder.AddToolBarButton(FDataprepEditorCommands::Get().CommitWorld);
+			}
+			ToolbarBuilder.EndSection();
+		}
+	};
+
+	TSharedPtr<FExtender> ToolbarExtender = MakeShareable(new FExtender);
+
+	FDataprepEditor* ThisEditor = this;
+
+	ToolbarExtender->AddToolBarExtension(
+		"Asset",
+		EExtensionHook::After,
+		ToolkitCommands,
+		FToolBarExtensionDelegate::CreateStatic(&Local::FillToolbar, ThisEditor)
+	);
+
+	AddToolbarExtender(ToolbarExtender);
+
+	IDataprepEditorModule* DataprepEditorModule = &FModuleManager::LoadModuleChecked<IDataprepEditorModule>("DataprepEditor");
+	AddToolbarExtender(DataprepEditorModule->GetToolBarExtensibilityManager()->GetAllExtenders(GetToolkitCommands(), GetEditingObjects()));
+}
+
+void FDataprepEditor::CreateTabs()
+{
+	AssetPreviewView = SNew(AssetPreviewWidget::SAssetsPreviewWidget);
+	AssetPreviewView->OnSelectionChanged().AddLambda(
+		[this](TSet< UObject* > Selection)
+		{
+			SetDetailsObjects(MoveTemp(Selection), false);
+		}
+	);
+
+	DataprepAssetView = SNew( SDataprepAssetView, DataprepAssetInterfacePtr.Get(), PipelineEditorCommands );
+
+	CreateScenePreviewTab();
+
+	// Create 3D viewport
+	SceneViewportView = SNew( SDataprepEditorViewport, SharedThis(this) )
+	.WorldToPreview( PreviewWorld.Get() );
+
+	// Create Details Panel
+	CreateDetailsViews();
+
+	// Temp code for the nodes development
+	if(DataprepRecipeBPPtr.IsValid())
+	{
+		// Create Pipeline Editor
+		CreatePipelineEditor();
+	}
+	// end of temp code for nodes development
+
+}
+
+// Temp code for the nodes development
+TSharedRef<SDockTab> FDataprepEditor::SpawnTabPipelineGraph(const FSpawnTabArgs & Args)
+{
+	check(Args.GetTabId() == PipelineGraphTabId);
+
+	if(!bIsDataprepInstance)
+	{
+		return SNew(SDockTab)
+			//.Icon(FDataprepEditorStyle::Get()->GetBrush("DataprepEditor.Tabs.Pipeline"))
+			.Label(LOCTEXT("DataprepEditor_PipelineTab_Title", "Pipeline"))
+			[
+				DataprepRecipeBPPtr.IsValid() ? PipelineView.ToSharedRef() : SNullWidget::NullWidget
+			];
+	}
+
+	return SNew(SDockTab);
+}
+// end of temp code for nodes development
+
+TSharedRef<SDockTab> FDataprepEditor::SpawnTabScenePreview(const FSpawnTabArgs & Args)
+{
+	check(Args.GetTabId() == ScenePreviewTabId);
+
+	return SNew(SDockTab)
+		//.Icon(FDataprepEditorStyle::Get()->GetBrush("DataprepEditor.Tabs.ScenePreview"))
+		.Label(LOCTEXT("DataprepEditor_ScenePreviewTab_Title", "Scene Preview"))
+		[
+			ScenePreviewView.ToSharedRef()
+		];
+}
+
+TSharedRef<SDockTab> FDataprepEditor::SpawnTabAssetPreview(const FSpawnTabArgs & Args)
+{
+	check(Args.GetTabId() == AssetPreviewTabId);
+
+	return SNew(SDockTab)
+		//.Icon(FDataprepEditorStyle::Get()->GetBrush("DataprepEditor.Tabs.AssetPreview"))
+		.Label(LOCTEXT("DataprepEditor_AssetPreviewTab_Title", "Asset Preview"))
+		[
+			SNew(SBorder)
+			.Padding(2.f)
+			.BorderImage(FEditorStyle::GetBrush("ToolPanel.GroupBorder"))
+			[
+				AssetPreviewView.ToSharedRef()
+			]
+		];
+}
+
+TSharedRef<SDockTab> FDataprepEditor::SpawnTabPalette(const FSpawnTabArgs & Args)
+{
+	check(Args.GetTabId() == PaletteTabId);
+
+	if(!bIsDataprepInstance)
+	{
+		return SNew(SDockTab)
+			.Icon(FSlateIcon(FEditorStyle::GetStyleSetName(), "Kismet.Tabs.Palette").GetIcon())
+			.Label(LOCTEXT("PaletteTab", "Palette"))
+			[
+				SNew(SDataprepPalette)
+			];
+	}
+
+	return SNew(SDockTab);
+}
+
+TSharedRef<SDockTab> FDataprepEditor::SpawnTabDataprep(const FSpawnTabArgs & Args)
+{
+	check(Args.GetTabId() == DataprepAssetTabId);
+
+	return SNew(SDockTab)
+	.Label(LOCTEXT("DataprepEditor_DataprepTab_Title", "Dataprep"))
+	[
+		SNew(SBorder)
+		.Padding(2.f)
+		.BorderImage(FEditorStyle::GetBrush("ToolPanel.GroupBorder"))
+		[
+			DataprepAssetView.ToSharedRef()
+		]
+	];
+}
+
+TSharedRef<SDockTab> FDataprepEditor::SpawnTabStatistics(const FSpawnTabArgs & Args)
+{
+	check(Args.GetTabId() == DataprepStatisticsTabId);
+
+	FStatsViewerModule& StatsViewerModule = FModuleManager::Get().LoadModuleChecked<FStatsViewerModule>("StatsViewer");
+	
+	const uint32 EnablePagesMask = (1 << EStatsPage::PrimitiveStats) | (1 << EStatsPage::StaticMeshLightingInfo) | (1 << EStatsPage::TextureStats);
+
+	return SNew(SDockTab)
+		.Label(LOCTEXT("DataprepEditor_StatisticsTab_Title", "Statistics"))
+		.Icon(FEditorStyle::GetBrush("LevelEditor.Tabs.StatsViewer"))
+		[
+			StatsViewerModule.CreateStatsViewer( *PreviewWorld.Get(), EnablePagesMask, TEXT("Dataprep") )
+		];
+}
+
+TSharedRef<SDockTab> FDataprepEditor::SpawnTabSceneViewport( const FSpawnTabArgs& Args )
+{
+	check( Args.GetTabId() == SceneViewportTabId );
+
+	TSharedRef<SDockTab> SpawnedTab =
+		SNew(SDockTab)
+		.Label( LOCTEXT("DataprepEditor_SceneViewportTab_Title", "Viewport") )
+		[
+			SceneViewportView.ToSharedRef()
+		];
+
+	return SpawnedTab;
+}
+
+TSharedRef<FTabManager::FLayout> FDataprepEditor::CreateDataprepLayout()
+{
+	return FTabManager::NewLayout("Standalone_DataprepEditor_Layout_v0.7")
+		->AddArea
+		(
+			FTabManager::NewPrimaryArea()->SetOrientation(Orient_Vertical)
+			->Split
+			(
+				FTabManager::NewStack()
+				->SetSizeCoefficient(0.1f)
+				->SetHideTabWell(true)
+				->AddTab(GetToolbarTabId(), ETabState::OpenedTab)
+			)
+			->Split
+			(
+				FTabManager::NewSplitter()->SetOrientation(Orient_Horizontal)
+				->Split
+				(
+					FTabManager::NewSplitter()->SetOrientation(Orient_Vertical)
+					->Split
+					(
+						FTabManager::NewSplitter()->SetOrientation(Orient_Horizontal)
+						->SetSizeCoefficient(0.75f)
+						->Split
+						(
+							FTabManager::NewStack()
+							->SetSizeCoefficient(0.2f)
+							->AddTab(AssetPreviewTabId, ETabState::OpenedTab)
+							->SetHideTabWell( true )
+						)
+						->Split
+						(
+							FTabManager::NewStack()
+							->SetSizeCoefficient(0.55f)
+							->AddTab(SceneViewportTabId, ETabState::OpenedTab)
+							->SetHideTabWell(true)
+						)
+						->Split
+						(
+							FTabManager::NewStack()
+							->SetSizeCoefficient(0.25f)
+							->AddTab(ScenePreviewTabId, ETabState::OpenedTab)
+							->SetHideTabWell( true )
+						)
+					)
+					->Split
+					(
+						FTabManager::NewSplitter()->SetOrientation(Orient_Horizontal)
+						->Split
+						(
+							FTabManager::NewStack()
+							->SetSizeCoefficient(0.15f)
+							->AddTab(PaletteTabId, ETabState::OpenedTab)
+							->SetHideTabWell( true )
+						)
+						// Temp code for the nodes development
+						->Split
+						(
+							FTabManager::NewStack()
+							->SetSizeCoefficient(0.85f)
+							->AddTab(PipelineGraphTabId, ETabState::OpenedTab)
+							->SetHideTabWell( true )
+						)
+						// end of temp code for nodes development
+					)
+				)
+				->Split
+				(
+					FTabManager::NewSplitter()->SetOrientation(Orient_Vertical)
+					->SetSizeCoefficient(0.25f)
+					->Split
+					(
+						FTabManager::NewStack()
+						->SetSizeCoefficient(0.3f)
+						->AddTab(DataprepAssetTabId, ETabState::OpenedTab)
+						->SetHideTabWell(true)
+					)
+					->Split
+					(
+						FTabManager::NewStack()
+						->SetSizeCoefficient(0.7f)
+						->AddTab(DetailsTabId, ETabState::OpenedTab)
+						->SetHideTabWell(true)
+					)
+				)
+			)
+		);
+}
+
+TSharedRef<FTabManager::FLayout> FDataprepEditor::CreateDataprepInstanceLayout()
+{
+	return FTabManager::NewLayout("Standalone_DataprepEditor_InstanceLayout_v0.2")
+		->AddArea
+		(
+			FTabManager::NewPrimaryArea()->SetOrientation(Orient_Vertical)
+			->Split
+			(
+				FTabManager::NewStack()
+				->SetSizeCoefficient(0.1f)
+				->SetHideTabWell(true)
+				->AddTab(GetToolbarTabId(), ETabState::OpenedTab)
+				// Don't want the secondary toolbar tab to be opened if there's nothing in it
+				//->AddTab(SecondaryToolbarTabId, ETabState::ClosedTab)
+			)
+			->Split
+			(
+				FTabManager::NewSplitter()->SetOrientation(Orient_Horizontal)
+				->Split
+				(
+					FTabManager::NewSplitter()->SetOrientation(Orient_Vertical)
+					->SetSizeCoefficient(0.2f)
+					->Split
+					(
+						FTabManager::NewStack()
+						->SetSizeCoefficient(0.5f)
+						->AddTab(ScenePreviewTabId, ETabState::OpenedTab)
+						->SetHideTabWell( true )
+					)
+					->Split
+					(
+						FTabManager::NewStack()
+						->SetSizeCoefficient(0.5f)
+						->AddTab(AssetPreviewTabId, ETabState::OpenedTab)
+						->SetHideTabWell( true )
+					)
+				)
+				->Split
+				(
+					FTabManager::NewSplitter()->SetOrientation(Orient_Vertical)
+					->SetSizeCoefficient(0.6f)
+					->Split
+					(
+						FTabManager::NewStack()
+						->SetSizeCoefficient(1.f)
+						->AddTab(SceneViewportTabId, ETabState::OpenedTab)
+						->SetHideTabWell(true)
+					)
+				)
+				->Split
+				(
+					FTabManager::NewSplitter()->SetOrientation(Orient_Vertical)
+					->SetSizeCoefficient(0.2f)
+					->Split
+					(
+						FTabManager::NewStack()
+						->SetSizeCoefficient(0.5f)
+						->AddTab(DataprepAssetTabId, ETabState::OpenedTab)
+						->SetHideTabWell(true)
+					)
+					->Split
+					(
+						FTabManager::NewStack()
+						->SetSizeCoefficient(0.5f)
+						->AddTab(DetailsTabId, ETabState::OpenedTab)
+						->SetHideTabWell(true)
+					)
+				)
+			)
+		);
+}
+
+void FDataprepEditor::UpdatePreviewPanels(bool bInclude3DViewport)
+{
+	TRACE_CPUPROFILER_EVENT_SCOPE(FDataprepEditor::UpdatePreviewPanels);
+
+	{
+		FTimeLogger TimeLogger( TEXT("UpdatePreviewPanels") );
+
+		AssetPreviewView->ClearAssetList();
+		FString SubstitutePath = DataprepAssetInterfacePtr->GetOutermost()->GetName();
+		if(DataprepAssetInterfacePtr->GetConsumer() != nullptr && !DataprepAssetInterfacePtr->GetConsumer()->GetTargetContentFolder().IsEmpty())
+		{
+			SubstitutePath = DataprepAssetInterfacePtr->GetConsumer()->GetTargetContentFolder();
+		}
+		AssetPreviewView->SetAssetsList( Assets, GetTransientContentFolder(), SubstitutePath );
+	}
+
+	if(bInclude3DViewport)
+	{
+		SceneViewportView->UpdateScene();
+	}
+}
+
+bool FDataprepEditor::OnRequestClose()
+{
+	const int ActorCount = PreviewWorld->GetActorCount();
+	if( bWorldBuilt && !bIgnoreCloseRequest && ActorCount > DefaultActorsInPreviewWorld.Num() )
+	{
+		// World was imported and is not empty -- show warning message
+		const FText Title( LOCTEXT( "DataprepEditor_ProceedWithClose", "Proceed with close") );
+		const FText Message( LOCTEXT( "DataprepEditor_ConfirmClose", "Imported data was not committed! Closing the editor will discard imported data.\nDo you want to close anyway?" ) );
+
+		return ( OpenMsgDlgInt( EAppMsgType::YesNo, Message, Title ) == EAppReturnType::Yes );
+	}
+	return !bIgnoreCloseRequest;
+}
+
+bool FDataprepEditor::CanBuildWorld()
+{
+	return DataprepAssetInterfacePtr->GetProducers()->GetProducersCount() > 0;
+}
+
+bool FDataprepEditor::CanExecutePipeline()
+{
+	return bWorldBuilt;
+}
+
+bool FDataprepEditor::CanCommitWorld()
+{
+	// Execution of pipeline is not required. User can directly commit result of import
+	return bWorldBuilt && DataprepAssetInterfacePtr->GetConsumer() != nullptr;
+}
+
+FString FDataprepEditor::GetTransientContentFolder()
+{
+	return FPaths::Combine( GetRootPackagePath(), FString::FromInt( FPlatformProcess::GetCurrentProcessId() ), SessionID );
+}
+
+bool FDataprepEditor::OnCanExecuteNextStep(UDataprepActionAsset* ActionAsset, UDataprepOperation* Operation, UDataprepFilter* Filter)
+{
+	// #ueent_todo: Make this action configurable by user
+	UpdatePreviewPanels(false);
+	return true;
+}
+
+void FDataprepEditor::OnActionsContextChanged( const UDataprepActionAsset* ActionAsset, bool bWorldChanged, bool bAssetsChanged, const TArray< TWeakObjectPtr<UObject> >& NewAssets )
+{
+	if(bAssetsChanged)
+	{
+		Assets = NewAssets;
+	}
+}
+
+void DataprepEditorUtil::DeleteTemporaryPackage( const FString& PathToDelete )
+{
+	// See ContentBrowserUtils::LoadAssetsIfNeeded
+	// See ContentBrowserUtils::DeleteFolders
+
+	FAssetRegistryModule& AssetRegistryModule = FModuleManager::Get().LoadModuleChecked<FAssetRegistryModule>(TEXT("AssetRegistry"));
+
+	// Form a filter from the path to delete
+	FARFilter Filter;
+	Filter.bRecursivePaths = true;
+	new (Filter.PackagePaths) FName( *PathToDelete );
+
+	// Query for a list of assets in the path to delete
+	TArray<FAssetData> AssetDataList;
+	AssetRegistryModule.Get().GetAssets(Filter, AssetDataList);
+
+	// Delete all registered objects which are in  memory and under this path
+	{
+		TArray<UObject*> AssetsToDelete;
+		AssetsToDelete.Reserve(AssetDataList.Num());
+		for(const FAssetData& AssetData : AssetDataList)
+		{
+			FSoftObjectPath ObjectPath( AssetData.ObjectPath.ToString() );
+
+			if(UObject* Object = ObjectPath.ResolveObject())
+			{
+				AssetsToDelete.Add(Object);
+			}
+		}
+
+		if(AssetsToDelete.Num() > 0)
+		{
+			ObjectTools::DeleteObjects(AssetsToDelete, false);
+		}
+	}
+
+	// Delete all assets not in memory but on disk
+	{
+		struct FEmptyFolderVisitor : public IPlatformFile::FDirectoryVisitor
+		{
+			bool bIsEmpty;
+
+			FEmptyFolderVisitor()
+				: bIsEmpty(true)
+			{
+			}
+
+			virtual bool Visit(const TCHAR* FilenameOrDirectory, bool bIsDirectory) override
+			{
+				if (!bIsDirectory)
+				{
+					bIsEmpty = false;
+					return false; // abort searching
+				}
+
+				return true; // continue searching
+			}
+		};
+
+		FString PathToDeleteOnDisk;
+		if (FPackageName::TryConvertLongPackageNameToFilename(PathToDelete, PathToDeleteOnDisk))
+		{
+			if(IFileManager::Get().DirectoryExists( *PathToDeleteOnDisk ))
+			{
+				// Look for files on disk in case the folder contains things not tracked by the asset registry
+				FEmptyFolderVisitor EmptyFolderVisitor;
+				IFileManager::Get().IterateDirectoryRecursively(*PathToDeleteOnDisk, EmptyFolderVisitor);
+
+				if (EmptyFolderVisitor.bIsEmpty && IFileManager::Get().DeleteDirectory(*PathToDeleteOnDisk, false, true))
+				{
+					AssetRegistryModule.Get().RemovePath(PathToDelete);
+				}
+			}
+			// Request deletion anyway
+			else
+			{
+				AssetRegistryModule.Get().RemovePath(PathToDelete);
+			}
+		}
+	}
+
+	// Check that no asset remains
+	AssetDataList.Reset();
+	AssetRegistryModule.Get().GetAssets(Filter, AssetDataList);
+	//ensure(AssetDataList.Num() == 0);
+}
+
+#undef LOCTEXT_NAMESPACE
