Trying to fix the case issue on linux

diff --git a/Engine/Plugins/Enterprise/DataprepEditor/Source/DataprepEditor/Private/DataprepEditor.h b/Engine/Plugins/Enterprise/DataprepEditor/Source/DataprepEditor/Private/DataprepEditor.h
new file mode 100644
--- /dev/null
+++ b/Engine/Plugins/Enterprise/DataprepEditor/Source/DataprepEditor/Private/DataprepEditor.h
@@ -0,0 +1,339 @@
+// Copyright 1998-2019 Epic Games, Inc. All Rights Reserved.
+
+#pragma once
+
+#include "CoreMinimal.h"
+
+#include "DataprepAssetInterface.h"
+#include "DataprepActionAsset.h"
+
+#include "EditorUndoClient.h"
+#include "Engine/EngineBaseTypes.h"
+#include "GraphEditor.h"
+#include "Misc/NotifyHook.h"
+#include "SceneOutlinerFwd.h"
+#include "Toolkits/AssetEditorToolkit.h"
+#include "Toolkits/IToolkitHost.h"
+#include "UObject/SoftObjectPath.h"
+#include "UObject/StrongObjectPtr.h"
+
+class FSpawnTabArgs;
+class FTabManager;
+class FUICommandList;
+class IMessageLogListing;
+class IMessageToken;
+class SDockableTab;
+class SGraphNodeDetailsWidget;
+class SInspectorView;
+class SWidget;
+class UWorld;
+class UEdGraphNode;
+
+namespace AssetPreviewWidget
+{
+	class SAssetsPreviewWidget;
+}
+
+/** Tuple linking an asset package path, a unique identifier and the UClass of the asset*/
+typedef TTuple< FString, UClass*, EObjectFlags > FSnapshotDataEntry;
+
+struct FDataprepSnapshot
+{
+	bool bIsValid;
+	TArray<FSnapshotDataEntry> DataEntries;
+
+	FDataprepSnapshot() : bIsValid(false) {}
+};
+
+typedef TTuple< UClass*, FText, FText > DataprepEditorClassDescription;
+
+class FDataprepEditor : public FAssetEditorToolkit, public FEditorUndoClient, public FNotifyHook
+{
+public:
+	FDataprepEditor();
+	virtual ~FDataprepEditor();
+
+	DECLARE_MULTICAST_DELEGATE(FOnDataprepAssetProducerChanged);
+
+	FOnDataprepAssetProducerChanged& OnDataprepAssetProducerChanged()
+	{
+		return DataprepAssetProducerChangedDelegate;
+	}
+
+	DECLARE_MULTICAST_DELEGATE(FOnDataprepAssetConsumerChanged);
+
+	FOnDataprepAssetConsumerChanged& OnDataprepAssetConsumerChanged()
+	{
+		return DataprepAssetConsumerChangedDelegate;
+	}
+
+	virtual void RegisterTabSpawners(const TSharedRef<FTabManager>& TabManager) override;
+	virtual void UnregisterTabSpawners(const TSharedRef<FTabManager>& TabManager) override;
+
+	void InitDataprepEditor(const EToolkitMode::Type Mode, const TSharedPtr< IToolkitHost >& InitToolkitHost, UDataprepAssetInterface* InDataprepAssetInterface, UObject* Blueprint = nullptr);
+
+	/** IToolkit interface */
+	virtual FName GetToolkitFName() const override;
+	virtual FText GetBaseToolkitName() const override;
+	virtual FString GetWorldCentricTabPrefix() const override;
+
+	/** @return Returns the color and opacity to use for the color that appears behind the tab text for this toolkit's tab in world-centric mode. */
+	virtual FLinearColor GetWorldCentricTabColorScale() const override;
+
+	UDataprepAssetInterface* GetDataprepAsset() const
+	{
+		return DataprepAssetInterfacePtr.Get();
+	}
+
+	/** Gets or sets the flag for context sensitivity in the graph action menu */
+	bool& GetIsContextSensitive() { return bIsActionMenuContextSensitive; }
+
+	/** Returns root package which all transient packages are created under */
+	static const FString& GetRootPackagePath();
+
+	/** Returns root directory which all transient directories and data are created under */
+	static const FString& GetRootTemporaryDir();
+
+	/** Return the selected object from the scene outliner / world preview */
+	const TSet<TWeakObjectPtr<UObject>>& GetWorldItemsSelection() const { return WorldItemsSelection; };
+
+	enum class EWorldSelectionFrom : uint8
+	{
+		SceneOutliner,
+		Viewport,
+		Unknow,
+	};
+
+	/** Set the selection of the world items */
+	void SetWorldObjectsSelection(TSet<TWeakObjectPtr<UObject>>&& NewSelection, EWorldSelectionFrom SelectionFrom = EWorldSelectionFrom::Unknow);
+
+private:
+	void BindCommands();
+	void OnSaveScene();
+	void OnBuildWorld();
+	void ResetBuildWorld();
+	void CleanPreviewWorld();
+	void OnExecutePipeline();
+	void OnCommitWorld();
+
+	/** Updates asset preview, scene outliner and 3D viewport if applicable */
+	void UpdatePreviewPanels(bool bInclude3dViewport = true);
+
+	void CreateTabs();
+
+	void CreateScenePreviewTab();
+
+	TSharedRef<SDockTab> SpawnTabAssetPreview(const FSpawnTabArgs& Args);
+	TSharedRef<SDockTab> SpawnTabScenePreview(const FSpawnTabArgs& Args);
+	TSharedRef<SDockTab> SpawnTabPalette(const FSpawnTabArgs& Args);
+	TSharedRef<SDockTab> SpawnTabDetails(const FSpawnTabArgs& Args);
+	TSharedRef<SDockTab> SpawnTabDataprep(const FSpawnTabArgs& Args);
+	TSharedRef<SDockTab> SpawnTabSceneViewport( const FSpawnTabArgs& Args );
+	TSharedRef<SDockTab> SpawnTabStatistics(const FSpawnTabArgs& Args);
+
+	TSharedRef<FTabManager::FLayout> CreateDataprepLayout();
+	TSharedRef<FTabManager::FLayout> CreateDataprepInstanceLayout();
+
+	void TryInvokingDetailsTab(bool bFlash);
+
+	/** Extends the toolbar menu to include static mesh editor options */
+	void ExtendMenu();
+
+	/** Builds the Data Prep Editor toolbar. */
+	void ExtendToolBar();
+
+	void CreateDetailsViews();
+
+	/** Gets a multi-cast delegate which is called whenever the graph object is changed to a different graph. */
+	//FOnGraphChanged& OnGraphChanged();
+
+	/** Called whenever the selected nodes in the graph editor changes. */
+	void OnPipelineEditorSelectionChanged(const TSet<UObject*>& SelectedNodes);
+
+	/** 
+	 * Change the objects display by the details panel
+	 * @param bCanEditProperties Should the details view let the user modify the objects.
+	 */
+	void SetDetailsObjects(const TSet<UObject*>& Objects, bool bCanEditProperties);
+
+	bool CanExecutePipeline();
+	bool CanBuildWorld();
+	bool CanCommitWorld();
+
+	virtual bool OnRequestClose() override;
+
+	/** Returns content folder under which all assets are stored after execution of all producers */
+	FString GetTransientContentFolder();
+
+	/** Create a snapshot of the world and tracked assets */
+	void TakeSnapshot();
+
+	/** Recreate preview world from snapshot */
+	void RestoreFromSnapshot( bool bUpdateViewport = false );
+
+	/** Handles changes in the Dataprep asset */
+	void OnDataprepAssetChanged(FDataprepAssetChangeType ChangeType);
+
+	/** Remove all temporary data remaining from previous runs of the Dataprep editor */
+	void CleanUpTemporaryDirectories();
+
+	/** Handles change to selection in SceneOutliner */
+	void OnSceneOutlinerSelectionChanged(SceneOutliner::FTreeItemPtr ItemPtr, ESelectInfo::Type SelectionMode);
+
+	bool OnCanExecuteNextStep(UDataprepActionAsset* ActionAsset, UDataprepOperation* Operation, UDataprepFilter* Filter );
+
+	/** Handles change to the content passed to an action */
+	void OnActionsContextChanged( const UDataprepActionAsset* ActionAsset, bool bWorldChanged, bool bAssetsChanged, const TArray< TWeakObjectPtr<UObject> >& NewAssets );
+
+private:
+	bool bWorldBuilt;
+	bool bIsFirstRun;
+	bool bPipelineChanged;
+	bool bIsDataprepInstance;
+	TWeakObjectPtr<UDataprepAssetInterface> DataprepAssetInterfacePtr;
+
+	FOnDataprepAssetProducerChanged DataprepAssetProducerChangedDelegate;
+	FOnDataprepAssetConsumerChanged DataprepAssetConsumerChangedDelegate;
+
+	TWeakPtr<SDockTab> DetailsTabPtr;
+	TSharedPtr<class SDataprepEditorViewport> SceneViewportView;
+	TSharedPtr<AssetPreviewWidget::SAssetsPreviewWidget> AssetPreviewView;
+	TSharedPtr<SWidget> ScenePreviewView;
+	TSharedPtr<SGraphNodeDetailsWidget> DetailsView;
+	TSharedPtr<class SDataprepAssetView > DataprepAssetView;
+	TSharedPtr<class SGraphEditor> PipelineView;
+
+	TSharedPtr<class ICustomSceneOutliner> SceneOutliner;
+
+	/** Command list for the pipeline editor */
+	TSharedPtr<FUICommandList> PipelineEditorCommands;
+	bool bIsActionMenuContextSensitive;
+	bool bSaveIntermediateBuildProducts;
+
+	/** 
+	 * Track the user selection from the scene outliner
+	 * This is referred to the garbage collector as a week reference
+	 */
+	TSet<TWeakObjectPtr<UObject>> WorldItemsSelection;
+
+	/**
+	 * All assets tracked for this editor.
+	 */
+	TArray<TWeakObjectPtr<UObject>> Assets;
+	TSet<TWeakObjectPtr<UObject>> CachedAssets;
+
+	/**
+	 * The world used to preview the inputs
+	 */
+	TStrongObjectPtr<UWorld> PreviewWorld;
+
+	TSet<class AActor*> DefaultActorsInPreviewWorld;
+
+	/** flag raised to prevent this editor to be closed */
+	bool bIgnoreCloseRequest;
+
+	/** Array of UClasses deriving from UDataprepContentConsumer */
+	TArray< DataprepEditorClassDescription > ConsumerDescriptions;
+
+	/** Temporary folder used to store content from snapshot */
+	FString TempDir;
+
+	/** Unique identifier assigned to each opened Dataprep editor to avoid name collision on cached data */
+	FString SessionID;
+
+	/** Structure to hold on the content of the latest call to OnBuildWorld */
+	FDataprepSnapshot ContentSnapshot;
+
+	/** Helper member to record classes of assets' sub-objects */
+	TMap<FString, UClass*> SnapshotClassesMap;
+
+	/**	The tab ids for all the tabs used */
+	static const FName ScenePreviewTabId;
+	static const FName AssetPreviewTabId;
+	static const FName PaletteTabId;
+	static const FName DetailsTabId;
+	static const FName DataprepAssetTabId;
+	static const FName SceneViewportTabId;
+	static const FName DataprepStatisticsTabId;
+
+//Temp Code to allow us to work on the nodes while we don't have our own graph.
+public:
+	UBlueprint* GetDataprepBlueprint()
+	{
+		return DataprepRecipeBPPtr.IsValid() ? DataprepRecipeBPPtr.Get() : nullptr;
+	}
+
+private:
+
+	TSharedRef<SDockTab> SpawnTabPipelineGraph(const FSpawnTabArgs& Args);
+
+	void CreatePipelineEditor();
+
+	
+	/** Called to create context menu when right-clicking on graph */
+	FActionMenuContent OnCreatePipelineActionMenu(UEdGraph* InGraph, const FVector2D& InNodePosition, const TArray<UEdGraphPin*>& InDraggedPins, bool bAutoExpand, SGraphEditor::FActionMenuClosed InOnMenuClosed);
+
+	/** Called when the Pipeline Blueprint has changed */
+	void OnPipelineChanged(UBlueprint* InBlueprint);
+
+	/** Called when the Pipeline Blueprint has been compiled */
+	void OnPipelineCompiled(UBlueprint* InBlueprint);
+
+	/** Callback when a token is clicked on in the compiler results log */
+	void OnLogTokenClicked(const TSharedRef<IMessageToken>& Token);
+
+	/** Callback when properties have finished being handled */
+	//virtual void OnFinishedChangingProperties(const FPropertyChangedEvent& PropertyChangedEvent);
+
+	/** Called when Compile button is clicked */
+	void OnCompile();
+
+	/** Returns the current status icon for the blueprint being edited */
+	FSlateIcon GetPipelineCompileButtonImage() const;
+	FText GetPipelineCompileButtonTooltip() const;
+
+	TSet<UObject*> GetSelectedPipelineNodes() const;
+	UEdGraph* GetPipelineGraph() const;
+	bool IsPipelineEditable() const;
+
+	void OnRenameNode();
+	bool CanRenameNode() const;
+
+	void SelectAllNodes();
+	bool CanSelectAllNodes() const;
+
+	void DeleteSelectedPipelineNodes();
+	bool CanDeletePipelineNodes() const;
+
+	void CopySelectedNodes();
+	bool CanCopyNodes() const;
+
+	void CutSelectedNodes();
+	bool CanCutNodes() const;
+
+	void PasteNodes();
+	bool CanPasteNodes() const;
+
+	void DuplicateNodes();
+	bool CanDuplicateNodes() const;
+
+	void OnCreateComment();
+
+	void DeleteSelectedDuplicatableNodes();
+	void PasteNodesHere(class UEdGraph* DestinationGraph, const FVector2D& GraphLocation);
+
+	bool OnNodeVerifyTitleCommit(const FText& NewText, UEdGraphNode* NodeBeingChanged, FText& OutErrorMessage);
+	void OnNodeTitleCommitted(const FText& NewText, ETextCommit::Type CommitInfo, UEdGraphNode* NodeBeingChanged);
+
+	// Start of temp code
+	TWeakObjectPtr<UBlueprint> DataprepRecipeBPPtr;
+
+	/** Compiler results log, with the log listing that it reflects */
+	TSharedPtr<SWidget> CompilerResults;
+	TSharedPtr<IMessageLogListing> CompilerResultsListing;
+
+	TSharedPtr<FDataprepActionContext> ActionsContext;
+
+	static const FName PipelineGraphTabId;
+	// End of temp code
+};
diff --git a/Engine/Plugins/Enterprise/DataprepEditor/Source/DataprepEditor/Private/DataprepEditorLogCategory.cpp b/Engine/Plugins/Enterprise/DataprepEditor/Source/DataprepEditor/Private/DataprepEditorLogCategory.cpp
new file mode 100644
--- /dev/null
+++ b/Engine/Plugins/Enterprise/DataprepEditor/Source/DataprepEditor/Private/DataprepEditorLogCategory.cpp
@@ -0,0 +1,6 @@
+// Copyright 1998-2019 Epic Games, Inc. All Rights Reserved.
+
+#include "DataprepEditorLogCategory.h"
+
+DEFINE_LOG_CATEGORY(LogDataprepEditor)
+
diff --git a/Engine/Plugins/Enterprise/DataprepEditor/Source/DataprepEditor/Private/DataprepEditorLogCategory.h b/Engine/Plugins/Enterprise/DataprepEditor/Source/DataprepEditor/Private/DataprepEditorLogCategory.h
new file mode 100644
--- /dev/null
+++ b/Engine/Plugins/Enterprise/DataprepEditor/Source/DataprepEditor/Private/DataprepEditorLogCategory.h
@@ -0,0 +1,7 @@
+// Copyright 1998-2019 Epic Games, Inc. All Rights Reserved.
+
+#pragma once
+
+#include "Logging/LogMacros.h"
+
+DECLARE_LOG_CATEGORY_EXTERN(LogDataprepEditor, Log, All);
diff --git a/Engine/Plugins/Enterprise/DataprepEditor/Source/DataprepEditor/Private/DataprepEditorModule.cpp b/Engine/Plugins/Enterprise/DataprepEditor/Source/DataprepEditor/Private/DataprepEditorModule.cpp
new file mode 100644
--- /dev/null
+++ b/Engine/Plugins/Enterprise/DataprepEditor/Source/DataprepEditor/Private/DataprepEditorModule.cpp
@@ -0,0 +1,128 @@
+// Copyright 1998-2019 Epic Games, Inc. All Rights Reserved.
+
+#include "DataprepEditorModule.h"
+
+#include "AssetTypeActions_Dataprep.h"
+#include "AssetTypeActions_DataprepAsset.h"
+#include "AssetTypeActions_DataprepAssetInterface.h"
+#include "DataprepAssetProducers.h"
+#include "DataprepEditor.h"
+#include "DataprepEditorStyle.h"
+#include "Widgets/DataprepWidgets.h"
+#include "Widgets/SDataprepEditorViewport.h"
+#include "Widgets/SDataprepProducersWidget.h"
+
+#include "Developer/AssetTools/Public/IAssetTools.h"
+#include "Developer/AssetTools/Public/AssetToolsModule.h"
+#include "Misc/PackageName.h"
+#include "Modules/ModuleManager.h"
+#include "PropertyEditorModule.h"
+#include "UObject/StrongObjectPtr.h"
+
+#include "Widgets/SNullWidget.h"
+
+const FName DataprepEditorAppIdentifier = FName(TEXT("DataprepEditorApp"));
+
+EAssetTypeCategories::Type IDataprepEditorModule::DataprepCategoryBit;
+
+#define LOCTEXT_NAMESPACE "DataprepEditorModule"
+
+class FDataprepEditorModule : public IDataprepEditorModule
+{
+public:
+	/** IModuleInterface implementation */
+	virtual void StartupModule() override
+	{
+		FDataprepEditorStyle::Initialize();
+
+		MenuExtensibilityManager = MakeShareable(new FExtensibilityManager);
+		ToolBarExtensibilityManager = MakeShareable(new FExtensibilityManager);
+
+		// Register asset type actions for DataPrepRecipe class
+		IAssetTools& AssetTools = FModuleManager::LoadModuleChecked<FAssetToolsModule>("AssetTools").Get();
+
+		// Register Dataprep category to group together asset type actions related to Dataprep
+		DataprepCategoryBit = AssetTools.RegisterAdvancedAssetCategory(FName(TEXT("Dataprep")), LOCTEXT("DataprepAssetCategory", "Dataprep"));
+
+		TSharedPtr<FAssetTypeActions_DataprepAssetInterface> DataprepAssetInterfaceTypeAction = MakeShared<FAssetTypeActions_DataprepAssetInterface>();
+		AssetTools.RegisterAssetTypeActions(DataprepAssetInterfaceTypeAction.ToSharedRef());
+		AssetTypeActionsArray.Add(DataprepAssetInterfaceTypeAction);
+
+		TSharedPtr<FAssetTypeActions_Dataprep> DataprepRecipeTypeAction = MakeShareable(new FAssetTypeActions_Dataprep);
+		AssetTools.RegisterAssetTypeActions(DataprepRecipeTypeAction.ToSharedRef());
+		AssetTypeActionsArray.Add(DataprepRecipeTypeAction);
+
+		TSharedPtr<FAssetTypeActions_DataprepAsset> DataprepAssetTypeAction = MakeShareable(new FAssetTypeActions_DataprepAsset);
+		AssetTools.RegisterAssetTypeActions(DataprepAssetTypeAction.ToSharedRef());
+		AssetTypeActionsArray.Add(DataprepAssetTypeAction);
+
+		TSharedPtr<FAssetTypeActions_DataprepAssetInstance> DataprepAssetInstanceTypeAction = MakeShareable(new FAssetTypeActions_DataprepAssetInstance);
+		AssetTools.RegisterAssetTypeActions(DataprepAssetInstanceTypeAction.ToSharedRef());
+		AssetTypeActionsArray.Add(DataprepAssetInstanceTypeAction);
+
+		// Register mount point for Dataprep editors root package folder
+		FPackageName::RegisterMountPoint( FDataprepEditor::GetRootPackagePath() + TEXT("/"), FDataprepEditor::GetRootTemporaryDir() );
+
+		// Register the details customizer
+		FPropertyEditorModule& PropertyModule = FModuleManager::LoadModuleChecked< FPropertyEditorModule >( TEXT("PropertyEditor") );
+		PropertyModule.RegisterCustomClassLayout( UDataprepAssetProducers::StaticClass()->GetFName(), FOnGetDetailCustomizationInstance::CreateStatic( &FDataprepAssetProducersDetails::MakeDetails ) );
+
+		SDataprepEditorViewport::LoadDefaultSettings();
+	}
+
+	virtual void ShutdownModule() override
+	{
+		MenuExtensibilityManager.Reset();
+		ToolBarExtensibilityManager.Reset();
+
+		// Unregister asset type actions
+		if (FModuleManager::Get().IsModuleLoaded(TEXT("AssetTools")))
+		{
+			IAssetTools& AssetTools = FModuleManager::LoadModuleChecked<FAssetToolsModule>("AssetTools").Get();
+			for (TSharedPtr<FAssetTypeActions_Base>& AssetTypeActions : AssetTypeActionsArray)
+			{
+				AssetTools.UnregisterAssetTypeActions(AssetTypeActions.ToSharedRef());
+			}
+		}
+		AssetTypeActionsArray.Empty();
+
+		FDataprepEditorStyle::Shutdown();
+
+		// Unregister mount point for Dataprep editors root package folder
+		FPackageName::UnRegisterMountPoint( FDataprepEditor::GetRootPackagePath() + TEXT("/"), FDataprepEditor::GetRootTemporaryDir() );
+
+		// Register the details customizer
+		FPropertyEditorModule& PropertyModule = FModuleManager::LoadModuleChecked< FPropertyEditorModule >( TEXT("PropertyEditor") );
+		PropertyModule.UnregisterCustomClassLayout( TEXT("DataprepAssetProducers") );
+	}
+
+	/** Gets the extensibility managers for outside entities to extend datasmith data prep editor's menus and toolbars */
+	virtual TSharedPtr<FExtensibilityManager> GetMenuExtensibilityManager() override { return MenuExtensibilityManager; }
+	virtual TSharedPtr<FExtensibilityManager> GetToolBarExtensibilityManager() override { return ToolBarExtensibilityManager; }
+
+	virtual TSharedRef<SWidget> CreateDataprepProducersWidget(UDataprepAssetProducers* AssetProducers) override
+	{
+		TSharedPtr<FUICommandList> CommandList = MakeShareable( new FUICommandList );
+		return AssetProducers ? SNew(SDataprepProducersWidget, AssetProducers, CommandList) : SNullWidget::NullWidget;
+	}
+
+	virtual TSharedRef<SWidget> CreateDataprepDetailsView(UObject* ObjectToDetail) override
+	{
+		if(ObjectToDetail)
+		{
+			return SNew(SDataprepDetailsView).Object( ObjectToDetail );
+		}
+
+		return SNullWidget::NullWidget;
+	}
+
+
+private:
+	TSharedPtr<FExtensibilityManager> MenuExtensibilityManager;
+	TSharedPtr<FExtensibilityManager> ToolBarExtensibilityManager;
+	TArray<TSharedPtr<FAssetTypeActions_Base>> AssetTypeActionsArray;
+};
+
+IMPLEMENT_MODULE(FDataprepEditorModule, DataprepEditor)
+
+#undef LOCTEXT_NAMESPACE
diff --git a/Engine/Plugins/Enterprise/DataprepEditor/Source/DataprepEditor/Private/DataprepPipelineEditor.cpp b/Engine/Plugins/Enterprise/DataprepEditor/Source/DataprepEditor/Private/DataprepPipelineEditor.cpp
new file mode 100644
--- /dev/null
+++ b/Engine/Plugins/Enterprise/DataprepEditor/Source/DataprepEditor/Private/DataprepPipelineEditor.cpp
@@ -0,0 +1,770 @@
+// Copyright 1998-2019 Epic Games, Inc. All Rights Reserved.
+
+#include "DataprepEditor.h"
+
+#include "BlueprintNodes/K2Node_DataprepAction.h"
+#include "DataprepEditorModule.h"
+#include "DataprepRecipe.h"
+#include "DataprepEditorActions.h"
+#include "SchemaActions/DataprepAllMenuActionCollector.h"
+#include "SchemaActions/IDataprepMenuActionCollector.h"
+#include "Widgets/SDataprepActionMenu.h"
+
+#include "BlueprintEditorSettings.h"
+#include "Components/ActorComponent.h"
+#include "EdGraph/EdGraphSchema.h"
+#include "EdGraphNode_Comment.h"
+#include "EdGraphSchema_K2_Actions.h"
+#include "EdGraphUtilities.h"
+#include "EditorStyleSet.h"
+#include "Engine/SCS_Node.h"
+#include "Framework/Application/SlateApplication.h"
+#include "Framework/Commands/GenericCommands.h"
+#include "Framework/Commands/UICommandList.h"
+#include "GraphEditor.h"
+#include "GraphEditorActions.h"
+#include "HAL/PlatformApplicationMisc.h"
+#include "IMessageLogListing.h"
+#include "K2Node_AddComponent.h"
+#include "K2Node_Event.h"
+#include "Kismet2/BlueprintEditorUtils.h"
+#include "Kismet2/CompilerResultsLog.h"
+#include "Kismet2/Kismet2NameValidators.h"
+#include "Kismet2/KismetEditorUtilities.h"
+#include "Logging/MessageLog.h"
+#include "Misc/UObjectToken.h"
+#include "PropertyEditorDelegates.h"
+#include "PropertyEditorModule.h"
+#include "SGraphNode.h"
+#include "SKismetInspector.h"
+#include "ScopedTransaction.h"
+#include "Widgets/Docking/SDockTab.h"
+#include "Widgets/Images/SImage.h"
+#include "Widgets/Input/SButton.h"
+#include "Widgets/Input/SEditableTextBox.h"
+#include "Widgets/Layout/Anchors.h"
+#include "Widgets/Layout/SBox.h"
+#include "Widgets/Layout/SConstraintCanvas.h"
+#include "Widgets/Notifications/SErrorText.h"
+
+#define LOCTEXT_NAMESPACE "DataprepPipelineEditor"
+
+namespace DataprepEditorUtils
+{
+	/**
+	 * Searches through a blueprint, looking for the most severe error'ing node.
+	 *
+	 * @param  Blueprint	The blueprint to search through.
+	 * @param  Severity		Defines the severity of the error/warning to search for.
+	 * @return The first node found with the specified error.
+	 * Inspired by BlueprintEditorImpl::FindNodeWithError
+	 */
+	UEdGraphNode* FindNodeWithError(FCompilerResultsLog const& ErrorLog, EMessageSeverity::Type Severity = EMessageSeverity::Error);
+}
+
+void FDataprepEditor::CreatePipelineEditor()
+{
+	if (!PipelineEditorCommands.IsValid())
+	{
+		PipelineEditorCommands = MakeShareable(new FUICommandList);
+
+		PipelineEditorCommands->MapAction(FGenericCommands::Get().Rename,
+			FExecuteAction::CreateSP(this, &FDataprepEditor::OnRenameNode),
+			FCanExecuteAction::CreateSP(this, &FDataprepEditor::CanRenameNode)
+		);
+
+		PipelineEditorCommands->MapAction(FGenericCommands::Get().SelectAll,
+			FExecuteAction::CreateSP(this, &FDataprepEditor::SelectAllNodes),
+			FCanExecuteAction::CreateSP(this, &FDataprepEditor::CanSelectAllNodes)
+		);
+
+		PipelineEditorCommands->MapAction(FGenericCommands::Get().Delete,
+			FExecuteAction::CreateSP(this, &FDataprepEditor::DeleteSelectedPipelineNodes),
+			FCanExecuteAction::CreateSP(this, &FDataprepEditor::CanDeletePipelineNodes)
+		);
+
+		PipelineEditorCommands->MapAction(FGenericCommands::Get().Copy,
+			FExecuteAction::CreateSP(this, &FDataprepEditor::CopySelectedNodes),
+			FCanExecuteAction::CreateSP(this, &FDataprepEditor::CanCopyNodes)
+		);
+
+		PipelineEditorCommands->MapAction(FGenericCommands::Get().Cut,
+			FExecuteAction::CreateSP(this, &FDataprepEditor::CutSelectedNodes),
+			FCanExecuteAction::CreateSP(this, &FDataprepEditor::CanCutNodes)
+		);
+
+		PipelineEditorCommands->MapAction(FGenericCommands::Get().Paste,
+			FExecuteAction::CreateSP(this, &FDataprepEditor::PasteNodes),
+			FCanExecuteAction::CreateSP(this, &FDataprepEditor::CanPasteNodes)
+		);
+
+		PipelineEditorCommands->MapAction(FGenericCommands::Get().Duplicate,
+			FExecuteAction::CreateSP(this, &FDataprepEditor::DuplicateNodes),
+			FCanExecuteAction::CreateSP(this, &FDataprepEditor::CanDuplicateNodes)
+		);
+
+		PipelineEditorCommands->MapAction(FGraphEditorCommands::Get().CreateComment,
+			FExecuteAction::CreateSP(this, &FDataprepEditor::OnCreateComment)
+		);
+	}
+
+	FGraphAppearanceInfo AppearanceInfo;
+	AppearanceInfo.CornerText = LOCTEXT("AppearanceCornerText", "DATAPREP");
+
+	// Create the title bar widget
+	TSharedRef<SWidget> TitleBarWidget =
+		SNew(SBorder)
+		.BorderImage(FEditorStyle::GetBrush(TEXT("Graph.TitleBackground")))
+		.HAlign(HAlign_Fill)
+		.Padding( 4.f )
+		[
+			SNew(SConstraintCanvas)
+			+ SConstraintCanvas::Slot()
+			.Anchors( FAnchors(0.5) )
+			.Alignment( FVector2D(0.5,0.5) )
+			.AutoSize( true )
+			[
+				SNew(STextBlock)
+				.Text(LOCTEXT("DataprepPipelineEditor_TitleBar_Label", "Dataprep Graph"))
+				.TextStyle(FEditorStyle::Get(), TEXT("GraphBreadcrumbButtonText"))
+			]
+		];
+
+	SGraphEditor::FGraphEditorEvents Events;
+	Events.OnSelectionChanged = SGraphEditor::FOnSelectionChanged::CreateSP(this, &FDataprepEditor::OnPipelineEditorSelectionChanged);
+	Events.OnCreateActionMenu = SGraphEditor::FOnCreateActionMenu::CreateSP(this, &FDataprepEditor::OnCreatePipelineActionMenu);
+	Events.OnVerifyTextCommit = FOnNodeVerifyTextCommit::CreateSP(this, &FDataprepEditor::OnNodeVerifyTitleCommit);
+	Events.OnTextCommitted = FOnNodeTextCommitted::CreateSP(this, &FDataprepEditor::OnNodeTitleCommitted);
+
+	UEdGraph* EventGraph = FBlueprintEditorUtils::FindEventGraph(DataprepRecipeBPPtr.Get());
+
+	PipelineView = SNew(SGraphEditor)
+		.AdditionalCommands(PipelineEditorCommands)
+		.Appearance(AppearanceInfo)
+		.TitleBar(TitleBarWidget)
+		.GraphToEdit(EventGraph)
+		.GraphEvents(Events);
+
+	DataprepRecipeBPPtr->OnChanged().AddSP(this, &FDataprepEditor::OnPipelineChanged);
+	DataprepRecipeBPPtr->OnCompiled().AddSP(this, &FDataprepEditor::OnPipelineCompiled);
+
+	CompilerResultsListing = FCompilerResultsLog::GetBlueprintMessageLog(DataprepRecipeBPPtr.Get());
+	CompilerResultsListing->OnMessageTokenClicked().AddSP(this, &FDataprepEditor::OnLogTokenClicked);
+
+}
+
+FActionMenuContent FDataprepEditor::OnCreatePipelineActionMenu(UEdGraph* InGraph, const FVector2D& InNodePosition, const TArray<UEdGraphPin*>& InDraggedPins, bool bAutoExpand, SGraphEditor::FActionMenuClosed InOnMenuClosed)
+{
+	// bAutoExpand is voluntary ignored for now
+	TUniquePtr<IDataprepMenuActionCollector> ActionCollector = MakeUnique<FDataprepAllMenuActionCollector>();
+
+	TSharedRef<SDataprepActionMenu> ActionMenu =
+		SNew (SDataprepActionMenu, MoveTemp(ActionCollector))
+			.TransactionText(LOCTEXT("AddingANewActionNode","Add a Action Node"))
+			.GraphObj(InGraph)
+			.NewNodePosition(InNodePosition)
+			.DraggedFromPins(InDraggedPins)
+			.OnClosedCallback(InOnMenuClosed);
+
+	return FActionMenuContent( ActionMenu, ActionMenu->GetFilterTextBox() );
+}
+
+bool FDataprepEditor::IsPipelineEditable() const
+{
+	return !FSlateApplication::Get().InKismetDebuggingMode() &&
+		   (DataprepRecipeBPPtr.IsValid()/* && DataprepRecipeBPPtr->CanRecompileWhilePlayingInEditor()*/) &&
+		   !FBlueprintEditorUtils::IsGraphReadOnly(GetPipelineGraph());
+}
+
+UEdGraph* FDataprepEditor::GetPipelineGraph() const
+{
+	return PipelineView.IsValid() ? PipelineView->GetCurrentGraph() : nullptr;
+}
+
+TSet<UObject*> FDataprepEditor::GetSelectedPipelineNodes() const
+{
+	TSet<UObject*> CurrentSelection;
+	if (PipelineView.IsValid())
+	{
+		CurrentSelection = PipelineView->GetSelectedNodes();
+	}
+	return CurrentSelection;
+}
+
+// Copied from FBlueprintEditor::DeleteSelectedNodes
+void FDataprepEditor::DeleteSelectedPipelineNodes()
+{
+	if (!PipelineView.IsValid())
+	{
+		return;
+	}
+
+	// DATAPREP_TODO: Add undo/redo
+	FScopedTransaction Transaction(FGenericCommands::Get().Delete->GetDescription());
+	PipelineView->GetCurrentGraph()->Modify();
+
+	bool bDidSomeModification = false;
+	const FGraphPanelSelectionSet SelectedNodes = GetSelectedPipelineNodes();
+
+	for (FGraphPanelSelectionSet::TConstIterator NodeIt(SelectedNodes); NodeIt; ++NodeIt)
+	{
+		if (UEdGraphNode* Node = Cast<UEdGraphNode>(*NodeIt))
+		{
+			if (Node->CanUserDeleteNode())
+			{
+				FBlueprintEditorUtils::RemoveNode(DataprepRecipeBPPtr.Get(), Node, true);
+				bDidSomeModification = true;
+			}
+		}
+	}
+
+	if (bDidSomeModification)
+	{
+		FBlueprintEditorUtils::MarkBlueprintAsModified(DataprepRecipeBPPtr.Get());
+	}
+	else
+	{
+		Transaction.Cancel();
+	}
+
+	//@TODO: Reselect items that were not deleted
+}
+
+// Temporary implementation from the blueprint editor
+void FDataprepEditor::OnRenameNode()
+{
+	if (PipelineView.IsValid())
+	{
+		const FGraphPanelSelectionSet SelectedNodes = GetSelectedPipelineNodes();
+		for (FGraphPanelSelectionSet::TConstIterator NodeIt(SelectedNodes); NodeIt; ++NodeIt)
+		{
+			UEdGraphNode* SelectedNode = Cast<UEdGraphNode>(*NodeIt);
+			if (SelectedNode != NULL && SelectedNode->bCanRenameNode)
+			{
+				PipelineView->IsNodeTitleVisible(SelectedNode, true);
+				break;
+			}
+		}
+	}
+}
+
+bool FDataprepEditor::CanRenameNode() const
+{
+	if ( IsPipelineEditable() )
+	{
+		const FGraphPanelSelectionSet SelectedNodes = GetSelectedPipelineNodes();
+		if ( SelectedNodes.Num() == 1)
+		{
+			if ( UEdGraphNode* SelectedNode = Cast<UEdGraphNode>(*SelectedNodes.CreateConstIterator()) )
+			{
+				return SelectedNode->bCanRenameNode;
+			}
+		}
+	}
+
+	return false;
+}
+
+void FDataprepEditor::SelectAllNodes()
+{
+	if (PipelineView.IsValid())
+	{
+		PipelineView->SelectAllNodes();
+	}
+}
+
+bool FDataprepEditor::CanSelectAllNodes() const
+{
+	return true;
+}
+
+bool FDataprepEditor::CanDeletePipelineNodes() const
+{
+	const FGraphPanelSelectionSet SelectedNodes = GetSelectedPipelineNodes();
+
+	bool bCanUserDeleteNode = false;
+
+	if (IsPipelineEditable() && SelectedNodes.Num() > 0)
+	{
+		for (UObject* NodeObject : SelectedNodes)
+		{
+			// If any nodes allow deleting, then do not disable the delete option
+			UEdGraphNode* Node = Cast<UEdGraphNode>(NodeObject);
+			if (Node->CanUserDeleteNode())
+			{
+				bCanUserDeleteNode = true;
+				break;
+			}
+		}
+	}
+
+	return bCanUserDeleteNode;
+}
+
+void FDataprepEditor::CopySelectedNodes()
+{
+	// Export the selected nodes and place the text on the clipboard
+	const FGraphPanelSelectionSet SelectedNodes = GetSelectedPipelineNodes();
+
+	FString ExportedText;
+
+	for (FGraphPanelSelectionSet::TConstIterator SelectedIter(SelectedNodes); SelectedIter; ++SelectedIter)
+	{
+		if (UEdGraphNode* Node = Cast<UEdGraphNode>(*SelectedIter))
+		{
+			Node->PrepareForCopying();
+		}
+	}
+
+	FEdGraphUtilities::ExportNodesToText(SelectedNodes, /*out*/ ExportedText);
+	FPlatformApplicationMisc::ClipboardCopy(*ExportedText);
+}
+
+bool FDataprepEditor::CanCopyNodes() const
+{
+	// If any of the nodes can be duplicated then we should allow copying
+	const FGraphPanelSelectionSet SelectedNodes = GetSelectedPipelineNodes();
+	for (FGraphPanelSelectionSet::TConstIterator SelectedIter(SelectedNodes); SelectedIter; ++SelectedIter)
+	{
+		UEdGraphNode* Node = Cast<UEdGraphNode>(*SelectedIter);
+		if ((Node != NULL) && Node->CanDuplicateNode())
+		{
+			return true;
+		}
+	}
+	return false;
+}
+
+void FDataprepEditor::CutSelectedNodes()
+{
+	CopySelectedNodes();
+	// Cut should only delete nodes that can be duplicated
+	DeleteSelectedDuplicatableNodes();
+}
+
+bool FDataprepEditor::CanCutNodes() const
+{
+	return IsPipelineEditable() && CanCopyNodes() && CanDeletePipelineNodes();
+}
+
+void FDataprepEditor::PasteNodes()
+{
+	// Find the graph editor with focus
+
+	if (!PipelineView.IsValid())
+	{
+		return;
+	}
+
+	PasteNodesHere(PipelineView->GetCurrentGraph(), PipelineView->GetPasteLocation());
+}
+
+bool FDataprepEditor::CanPasteNodes() const
+{
+	// Find the graph editor with focus
+	if (!PipelineView.IsValid() && !IsPipelineEditable() )
+	{
+		return false;
+	}
+
+	FString ClipboardContent;
+	FPlatformApplicationMisc::ClipboardPaste(ClipboardContent);
+
+	return FEdGraphUtilities::CanImportNodesFromText(PipelineView->GetCurrentGraph(), ClipboardContent);
+}
+
+void FDataprepEditor::DuplicateNodes()
+{
+	// Copy and paste current selection
+	CopySelectedNodes();
+	PasteNodes();
+}
+
+bool FDataprepEditor::CanDuplicateNodes() const
+{
+	return IsPipelineEditable() && CanCopyNodes();
+}
+
+void FDataprepEditor::DeleteSelectedDuplicatableNodes()
+{
+	// Cache off the old selection
+	const FGraphPanelSelectionSet OldSelectedNodes = GetSelectedPipelineNodes();
+
+	// Clear the selection and only select the nodes that can be duplicated
+	FGraphPanelSelectionSet CurrentSelection;
+	if (PipelineView.IsValid())
+	{
+		PipelineView->ClearSelectionSet();
+
+		FGraphPanelSelectionSet RemainingNodes;
+		for (FGraphPanelSelectionSet::TConstIterator SelectedIter(OldSelectedNodes); SelectedIter; ++SelectedIter)
+		{
+			UEdGraphNode* Node = Cast<UEdGraphNode>(*SelectedIter);
+			if ((Node != NULL) && Node->CanDuplicateNode())
+			{
+				PipelineView->SetNodeSelection(Node, true);
+			}
+			else
+			{
+				RemainingNodes.Add(Node);
+			}
+		}
+
+		// Delete the duplicatable nodes
+		DeleteSelectedPipelineNodes();
+
+		// Reselect whatever's left from the original selection after the deletion
+		PipelineView->ClearSelectionSet();
+
+		for (FGraphPanelSelectionSet::TConstIterator SelectedIter(RemainingNodes); SelectedIter; ++SelectedIter)
+		{
+			if (UEdGraphNode* Node = Cast<UEdGraphNode>(*SelectedIter))
+			{
+				PipelineView->SetNodeSelection(Node, true);
+			}
+		}
+	}
+}
+
+void FDataprepEditor::PasteNodesHere(class UEdGraph* DestinationGraph, const FVector2D& GraphLocation)
+{
+	// Find the graph editor with focus
+	if (!PipelineView.IsValid())
+	{
+		return;
+	}
+	// Select the newly pasted stuff
+	bool bNeedToModifyStructurally = false;
+	{
+		const FScopedTransaction Transaction(FGenericCommands::Get().Paste->GetDescription());
+		DestinationGraph->Modify();
+		PipelineView->ClearSelectionSet();
+
+		// Grab the text to paste from the clipboard.
+		FString TextToImport;
+		FPlatformApplicationMisc::ClipboardPaste(TextToImport);
+
+		// Import the nodes
+		TSet<UEdGraphNode*> PastedNodes;
+		FEdGraphUtilities::ImportNodesFromText(DestinationGraph, TextToImport, /*out*/ PastedNodes);
+
+		//Average position of nodes so we can move them while still maintaining relative distances to each other
+		FVector2D AvgNodePosition(0.0f, 0.0f);
+
+		for (TSet<UEdGraphNode*>::TIterator It(PastedNodes); It; ++It)
+		{
+			UEdGraphNode* Node = *It;
+			AvgNodePosition.X += Node->NodePosX;
+			AvgNodePosition.Y += Node->NodePosY;
+		}
+
+		float InvNumNodes = 1.0f / float(PastedNodes.Num());
+		AvgNodePosition.X *= InvNumNodes;
+		AvgNodePosition.Y *= InvNumNodes;
+
+		for (TSet<UEdGraphNode*>::TIterator It(PastedNodes); It; ++It)
+		{
+			UEdGraphNode* Node = *It;
+			PipelineView->SetNodeSelection(Node, true);
+
+			Node->NodePosX = (Node->NodePosX - AvgNodePosition.X) + GraphLocation.X;
+			Node->NodePosY = (Node->NodePosY - AvgNodePosition.Y) + GraphLocation.Y;
+
+			Node->SnapToGrid(SNodePanel::GetSnapGridSize());
+
+			// Give new node a different Guid from the old one
+			Node->CreateNewGuid();
+
+			UK2Node* K2Node = Cast<UK2Node>(Node);
+			if (K2Node && K2Node->NodeCausesStructuralBlueprintChange())
+			{
+				bNeedToModifyStructurally = true;
+			}
+
+			// The only useful K2Node for this editor are our DataprepAction node
+			if (K2Node)
+			{
+				if (!K2Node->IsA<UK2Node_DataprepAction>())
+				{
+					FBlueprintEditorUtils::RemoveNode(DataprepRecipeBPPtr.Get(), Node, true);
+				}
+			}
+			
+		}
+	}
+
+	if (bNeedToModifyStructurally)
+	{
+		FBlueprintEditorUtils::MarkBlueprintAsStructurallyModified(DataprepRecipeBPPtr.Get());
+	}
+	else
+	{
+		FBlueprintEditorUtils::MarkBlueprintAsModified(DataprepRecipeBPPtr.Get());
+	}
+
+	// Update UI
+	PipelineView->NotifyGraphChanged();
+}
+
+void FDataprepEditor::OnCreateComment()
+{
+	if (PipelineView.IsValid())
+	{
+		if (UEdGraph* Graph = PipelineView->GetCurrentGraph())
+		{
+			if (const UEdGraphSchema* Schema = Graph->GetSchema())
+			{
+				// Add menu item for creating comment boxes
+				UEdGraphNode_Comment* CommentTemplate = NewObject<UEdGraphNode_Comment>();
+
+				FVector2D SpawnLocation = PipelineView->GetPasteLocation();
+
+				FSlateRect Bounds;
+				if ( PipelineView->GetBoundsForSelectedNodes(Bounds, 50.0f) )
+				{
+					CommentTemplate->SetBounds(Bounds);
+					SpawnLocation.X = CommentTemplate->NodePosX;
+					SpawnLocation.Y = CommentTemplate->NodePosY;
+				}
+
+				UEdGraphNode* NewNode = FEdGraphSchemaAction_NewNode::SpawnNodeFromTemplate<UEdGraphNode_Comment>(Graph, CommentTemplate, SpawnLocation, /** bSelectNewNode */ true);
+
+				// Mark Blueprint as structurally modified since
+				// UK2Node_Comment::NodeCausesStructuralBlueprintChange used to return true
+				FBlueprintEditorUtils::MarkBlueprintAsStructurallyModified(DataprepRecipeBPPtr.Get());
+			}
+		}
+	}
+}
+
+
+bool FDataprepEditor::OnNodeVerifyTitleCommit(const FText& NewText, UEdGraphNode* NodeBeingChanged, FText& OutErrorMessage)
+{
+	bool bValid(false);
+
+	if (NodeBeingChanged && NodeBeingChanged->bCanRenameNode)
+	{
+		// Clear off any existing error message 
+		NodeBeingChanged->ErrorMsg.Empty();
+		NodeBeingChanged->bHasCompilerMessage = false;
+
+		TSharedPtr<INameValidatorInterface> NameEntryValidator = NodeBeingChanged->MakeNameValidator();
+
+		check( NameEntryValidator.IsValid() );
+
+		EValidatorResult VResult = NameEntryValidator->IsValid(NewText.ToString(), true);
+		if (VResult == EValidatorResult::Ok)
+		{
+			bValid = true;
+		}
+		else if (PipelineView.IsValid())
+		{
+			EValidatorResult Valid = NameEntryValidator->IsValid(NewText.ToString(), false);
+
+			NodeBeingChanged->bHasCompilerMessage = true;
+			NodeBeingChanged->ErrorMsg = NameEntryValidator->GetErrorString(NewText.ToString(), Valid);
+			NodeBeingChanged->ErrorType = EMessageSeverity::Error;
+		}
+	}
+
+	return bValid;
+}
+
+void FDataprepEditor::OnNodeTitleCommitted(const FText& NewText, ETextCommit::Type CommitInfo, UEdGraphNode* NodeBeingChanged)
+{
+	if (NodeBeingChanged)
+	{
+		const FScopedTransaction Transaction(NSLOCTEXT("RenameNode", "RenameNode", "Rename Node"));
+		NodeBeingChanged->Modify();
+		NodeBeingChanged->OnRenameNode(NewText.ToString());
+	}
+}
+
+// Inspired by FBlueprintEditor::OnBlueprintCompiled
+void FDataprepEditor::OnPipelineCompiled(UBlueprint* InBlueprint)
+{
+	if (InBlueprint)
+	{
+		check(InBlueprint == DataprepRecipeBPPtr.Get());
+		UEdGraph* EventGraph = FBlueprintEditorUtils::FindEventGraph(DataprepRecipeBPPtr.Get());
+		for (const UEdGraphNode* Node : EventGraph->Nodes)
+		{
+			if (Node)
+			{
+				TSharedPtr<SGraphNode> Widget = Node->DEPRECATED_NodeWidget.Pin();
+				if (Widget.IsValid())
+				{
+					Widget->RefreshErrorInfo();
+				}
+			}
+		}
+
+		OnPipelineChanged(InBlueprint);
+	}
+}
+
+// Inspired by FBlueprintEditor::OnBlueprintChangedImpl
+void FDataprepEditor::OnPipelineChanged(UBlueprint* InBlueprint)
+{
+	if (InBlueprint)
+	{
+		check(InBlueprint == DataprepRecipeBPPtr.Get());
+
+		// Notify that the blueprint has been changed (update Content browser, etc)
+		InBlueprint->PostEditChange();
+
+		// Call PostEditChange() on any Actors that are based on this Blueprint 
+		FBlueprintEditorUtils::PostEditChangeBlueprintActors(InBlueprint);
+	}
+
+	if ( PipelineView )
+	{
+		PipelineView->NotifyGraphChanged();
+	}
+}
+
+// Inspired by FBlueprintEditor::Compile
+void FDataprepEditor::OnCompile()
+{
+	TRACE_CPUPROFILER_EVENT_SCOPE(FDataprepEditor::OnCompile);
+
+	if (DataprepRecipeBPPtr.IsValid())
+	{
+		FMessageLog BlueprintLog("BlueprintLog");
+
+		FFormatNamedArguments Arguments;
+		Arguments.Add(TEXT("BlueprintName"), FText::FromString(DataprepRecipeBPPtr->GetName()));
+		BlueprintLog.NewPage(FText::Format(LOCTEXT("CompilationPageLabel", "Compile {BlueprintName}"), Arguments));
+
+		FCompilerResultsLog LogResults;
+		LogResults.SetSourcePath(DataprepRecipeBPPtr->GetPathName());
+		LogResults.BeginEvent(TEXT("Compile"));
+		LogResults.bLogDetailedResults = GetDefault<UBlueprintEditorSettings>()->bShowDetailedCompileResults;
+		LogResults.EventDisplayThresholdMs = GetDefault<UBlueprintEditorSettings>()->CompileEventDisplayThresholdMs;
+		EBlueprintCompileOptions CompileOptions = EBlueprintCompileOptions::None;
+		if (bSaveIntermediateBuildProducts)
+		{
+			CompileOptions |= EBlueprintCompileOptions::SaveIntermediateProducts;
+		}
+		FKismetEditorUtilities::CompileBlueprint(DataprepRecipeBPPtr.Get(), CompileOptions, &LogResults);
+
+		LogResults.EndEvent();
+
+		const bool bForceMessageDisplay = ((LogResults.NumWarnings > 0) || (LogResults.NumErrors > 0)) && !DataprepRecipeBPPtr->bIsRegeneratingOnLoad;
+		{
+			CompilerResultsListing->ClearMessages();
+
+			// Note we don't mirror to the output log here as the compiler already does that
+			CompilerResultsListing->AddMessages(LogResults.Messages, false);
+
+			// DATAPREP_TODO: Create tab for compiling error
+			//if (!bEditorMarkedAsClosed && bForceMessageDisplay && GetCurrentMode() == FBlueprintEditorApplicationModes::StandardBlueprintEditorMode)
+			//{
+			//	TabManager->InvokeTab(FBlueprintEditorTabs::CompilerResultsID);
+			//}
+		}
+
+		UBlueprintEditorSettings const* BpEditorSettings = GetDefault<UBlueprintEditorSettings>();
+		if ((LogResults.NumErrors > 0) && BpEditorSettings->bJumpToNodeErrors)
+		{
+			if (UEdGraphNode* NodeWithError = DataprepEditorUtils::FindNodeWithError(LogResults))
+			{
+				if (PipelineView.IsValid())
+				{
+					PipelineView->JumpToNode(NodeWithError, false);
+				}
+			}
+		}
+
+		if (DataprepRecipeBPPtr->UpgradeNotesLog.IsValid())
+		{
+			CompilerResultsListing->AddMessages(DataprepRecipeBPPtr->UpgradeNotesLog->Messages);
+		}
+	}
+}
+
+// Copied from FBlueprintEditorToolbar::GetStatusImage
+FSlateIcon FDataprepEditor::GetPipelineCompileButtonImage() const
+{
+	switch (DataprepRecipeBPPtr->Status)
+	{
+		case BS_Error:
+			return FSlateIcon(FEditorStyle::GetStyleSetName(), "Kismet.Status.Error");
+		case BS_UpToDate:
+			return FSlateIcon(FEditorStyle::GetStyleSetName(), "Kismet.Status.Good");
+		case BS_UpToDateWithWarnings:
+			return FSlateIcon(FEditorStyle::GetStyleSetName(), "Kismet.Status.Warning");
+	}
+
+	return FSlateIcon(FEditorStyle::GetStyleSetName(), "Kismet.Status.Unknown");
+}
+
+// Copied from FBlueprintEditorToolbar::GetStatusTooltip
+FText FDataprepEditor::GetPipelineCompileButtonTooltip() const
+{
+	switch (DataprepRecipeBPPtr->Status)
+	{
+	case BS_Dirty:
+		return LOCTEXT("Dirty_Status", "Dataprep graph is dirty; needs to be recompiled");
+	case BS_Error:
+		return LOCTEXT("CompileError_Status", "There was an error during compilation, see Dataprep graph viewport for details");
+	case BS_UpToDate:
+		return LOCTEXT("GoodToGo_Status", "Dataprep is ready");
+	case BS_UpToDateWithWarnings:
+		return LOCTEXT("GoodToGoWarning_Status", "There was a warning during compilation, see Dataprep graph viewport for details");
+	}
+
+	return LOCTEXT("Recompile_Status", "Unknown status; should recompile Dataprep graph");
+}
+
+void FDataprepEditor::OnLogTokenClicked(const TSharedRef<IMessageToken>& Token)
+{
+	if (Token->GetType() == EMessageToken::Object)
+	{
+		const TSharedRef<FUObjectToken> ObjectToken = StaticCastSharedRef<FUObjectToken>(Token);
+		if (ObjectToken->GetObject().IsValid() && PipelineView.IsValid())
+		{
+			const UEdGraphNode* Node = Cast<const UEdGraphNode>(ObjectToken->GetObject().Get());
+			PipelineView->JumpToNode(Node, false);
+		}
+	}
+	else if (Token->GetType() == EMessageToken::EdGraph)
+	{
+		const TSharedRef<FEdGraphToken> EdGraphToken = StaticCastSharedRef<FEdGraphToken>(Token);
+		const UEdGraphPin* PinBeingReferenced = EdGraphToken->GetPin();
+		const UObject* ObjectBeingReferenced = EdGraphToken->GetGraphObject();
+		if (PinBeingReferenced && PipelineView.IsValid())
+		{
+			PipelineView->JumpToPin(PinBeingReferenced);
+		}
+		else if (ObjectBeingReferenced && PipelineView.IsValid())
+		{
+			const UEdGraphNode* Node = Cast<const UEdGraphNode>(ObjectBeingReferenced);
+			PipelineView->JumpToNode(Node, false);
+		}
+	}
+}
+
+UEdGraphNode* DataprepEditorUtils::FindNodeWithError(FCompilerResultsLog const& ErrorLog, EMessageSeverity::Type Severity/* = EMessageSeverity::Error*/)
+{
+	UEdGraphNode* ChoiceNode = nullptr;
+	for (TWeakObjectPtr<UEdGraphNode> NodePtr : ErrorLog.AnnotatedNodes)
+	{
+		UEdGraphNode* Node = NodePtr.Get();
+		if ((Node != nullptr) && (Node->ErrorType <= Severity))
+		{
+			if ((ChoiceNode == nullptr) || (Node->ErrorType < ChoiceNode->ErrorType))
+			{
+				ChoiceNode = Node;
+				if (ChoiceNode->ErrorType == 0)
+				{
+					break;
+				}
+			}
+		}
+	}
+
+	return ChoiceNode;
+}
+
+#undef LOCTEXT_NAMESPACE
diff --git a/Engine/Plugins/Enterprise/DataprepEditor/Source/DataprepEditor/Private/SchemaActions/DataprepDragDropOp.cpp b/Engine/Plugins/Enterprise/DataprepEditor/Source/DataprepEditor/Private/SchemaActions/DataprepDragDropOp.cpp
new file mode 100644
--- /dev/null
+++ b/Engine/Plugins/Enterprise/DataprepEditor/Source/DataprepEditor/Private/SchemaActions/DataprepDragDropOp.cpp
@@ -0,0 +1,172 @@
+// Copyright 1998-2019 Epic Games, Inc. All Rights Reserved.
+
+#include "SchemaActions/DataprepDragDropOp.h"
+
+#include "BlueprintNodes/K2Node_DataprepAction.h"
+#include "DataprepActionAsset.h"
+#include "DataprepSchemaActionUtils.h"
+
+#include "EdGraph/EdGraph.h"
+#include "EdGraphSchema_K2.h"
+#include "EditorStyleSet.h"
+#include "Kismet2/BlueprintEditorUtils.h"
+#include "ScopedTransaction.h"
+#include "Templates/UnrealTypeTraits.h"
+#include "Widgets/Text/STextBlock.h"
+
+#define LOCTEXT_NAMESPACE "DataprepDragAndDrop"
+
+FDataprepDragDropOp::FDataprepDragDropOp()
+	: FGraphEditorDragDropAction()
+	, HoveredDataprepActionContext()
+{
+	bDropTargetValid = false;
+	Construct();
+}
+
+TSharedRef<FDataprepDragDropOp> FDataprepDragDropOp::New(TSharedRef<FDataprepSchemaAction> InAction)
+{
+	TSharedRef< FDataprepDragDropOp > DragDrop = MakeShared< FDataprepDragDropOp >();
+	DragDrop->DataprepGraphOperation.BindSP( InAction, &FDataprepSchemaAction::ExecuteAction );
+	return DragDrop;
+}
+
+TSharedRef<FDataprepDragDropOp> FDataprepDragDropOp::New(FDataprepGraphOperation&& DataprepGraphOperation)
+{
+	TSharedRef< FDataprepDragDropOp > DragDrop = MakeShared< FDataprepDragDropOp >();
+	DragDrop->DataprepGraphOperation = MoveTemp( DataprepGraphOperation );
+	return DragDrop;
+}
+
+void FDataprepDragDropOp::SetHoveredDataprepActionContext(TOptional<FDataprepSchemaActionContext> Context)
+{
+	if ( HoveredDataprepActionContext != Context )
+	{
+		HoveredDataprepActionContext = Context;
+		HoverTargetChanged();
+	}
+}
+
+FReply FDataprepDragDropOp::DroppedOnDataprepActionContext(const FDataprepSchemaActionContext& Context)
+{
+	if ( DataprepPreDropConfirmation.IsBound() )
+	{
+		TFunction<void ()> OnConfirmation( [Operation = StaticCastSharedRef<FDataprepDragDropOp>(AsShared()), Context = FDataprepSchemaActionContext(Context)] ()
+			{
+				Operation->DoDropOnDataprepActionContext( Context );
+			} );
+
+		DataprepPreDropConfirmation.Execute( Context, OnConfirmation );
+	}
+	else
+	{
+		DoDropOnDataprepActionContext( Context );
+	}
+
+	return FReply::Handled();
+}
+
+void FDataprepDragDropOp::HoverTargetChanged()
+{
+	FText DrapDropText;
+	if ( HoveredDataprepActionContext )
+	{
+		bDropTargetValid = true;
+		DrapDropText = LOCTEXT("TargetIsDataprepActionContext", "Add a Step to Dataprep Action");
+	}
+	else if ( UEdGraph* EdGraph = GetHoveredGraph() )
+	{
+		if ( const UEdGraphSchema_K2* GraphSchema_k2 = Cast<UEdGraphSchema_K2>( EdGraph->GetSchema() ) )
+		{
+			bDropTargetValid = true;
+			DrapDropText = LOCTEXT("TargetIsBlueprintGraph", "Add a Dataprep Action");
+		}
+		else
+		{
+			bDropTargetValid = false;
+			DrapDropText = LOCTEXT("TargetGraphIsInvalid", "Can only be drop on a blueprint graph");
+		}
+	}
+	else
+	{
+		bDropTargetValid = false;
+		DrapDropText = FText::FromString( TEXT("Can't drop here") );
+	}
+
+	const FSlateBrush* Symbol = bDropTargetValid ? FEditorStyle::GetBrush( TEXT("Graph.ConnectorFeedback.OK") ) :  FEditorStyle::GetBrush( TEXT("Graph.ConnectorFeedback.Error") );
+	SetSimpleFeedbackMessage( Symbol, FLinearColor::White, DrapDropText );
+}
+
+FReply FDataprepDragDropOp::DroppedOnPanel(const TSharedRef<SWidget>& Panel, FVector2D ScreenPosition, FVector2D GraphPosition, UEdGraph& Graph)
+{
+	if ( bDropTargetValid )
+	{
+		if ( DataprepPreDropConfirmation.IsBound() )
+		{
+			TFunction<void ()> OnConfirmation ( [Operation = StaticCastSharedRef<FDataprepDragDropOp>( AsShared() ), Panel, ScreenPosition, GraphPosition, GraphPtr = TWeakObjectPtr<UEdGraph>(&Graph)] ()
+				{
+					UEdGraph* Graph = GraphPtr.Get();
+					if ( Graph )
+					{
+						Operation->DoDropOnPanel(Panel, ScreenPosition, GraphPosition, *Graph);
+					}
+				} );
+
+			DataprepPreDropConfirmation.Execute( FDataprepSchemaActionContext(), OnConfirmation);
+		
+			return FReply::Handled();
+		}
+		else
+		{
+			DoDropOnPanel( Panel, ScreenPosition, GraphPosition, Graph );
+			return FReply::Handled();
+		}
+	}
+
+	return FReply::Unhandled();
+}
+
+void FDataprepDragDropOp::SetPreDropConfirmation(FDataprepPreDropConfirmation && Confirmation)
+{
+	DataprepPreDropConfirmation = MoveTemp( Confirmation );
+}
+
+bool FDataprepDragDropOp::DoDropOnDataprepActionContext(const FDataprepSchemaActionContext& Context)
+{
+	if ( DataprepGraphOperation.IsBound() )
+	{
+		FScopedTransaction Transaction( LOCTEXT("AddStep", "Add a Step to a Dataprep Action") );
+		bool bDidModification = DataprepGraphOperation.Execute( Context );
+		if ( !bDidModification )
+		{
+			Transaction.Cancel();
+		}
+		return bDidModification;
+	}
+	return false;
+}
+
+void FDataprepDragDropOp::DoDropOnPanel(const TSharedRef<SWidget>& Panel, FVector2D ScreenPosition, FVector2D GraphPosition, UEdGraph& Graph)
+{
+	if (UEdGraph* EdGraph = GetHoveredGraph())
+	{
+		FScopedTransaction Transaction( LOCTEXT("AddNode", "Add Dataprep Action Node") );
+		UK2Node_DataprepAction* DataprepActionNode = DataprepSchemaActionUtils::SpawnEdGraphNode< UK2Node_DataprepAction >( Graph, GraphPosition );
+		check( DataprepActionNode );
+		DataprepActionNode->CreateDataprepActionAsset();
+		DataprepActionNode->AutowireNewNode( GetHoveredPin() );
+
+		FDataprepSchemaActionContext Context;
+		Context.DataprepActionPtr = DataprepActionNode->GetDataprepAction();
+		if ( !DoDropOnDataprepActionContext( Context ) )
+		{
+			Transaction.Cancel();
+		}
+
+		if ( UBlueprint* Blueprint = FBlueprintEditorUtils::FindBlueprintForGraphChecked( EdGraph ) )
+		{
+			FBlueprintEditorUtils::MarkBlueprintAsModified( Blueprint );
+		}
+	}
+}
+#undef LOCTEXT_NAMESPACE
diff --git a/Engine/Plugins/Enterprise/DataprepEditor/Source/DataprepEditor/Private/SchemaActions/DataprepFetcherMenuActionCollector.cpp b/Engine/Plugins/Enterprise/DataprepEditor/Source/DataprepEditor/Private/SchemaActions/DataprepFetcherMenuActionCollector.cpp
new file mode 100644
--- /dev/null
+++ b/Engine/Plugins/Enterprise/DataprepEditor/Source/DataprepEditor/Private/SchemaActions/DataprepFetcherMenuActionCollector.cpp
@@ -0,0 +1,60 @@
+// Copyright 1998-2019 Epic Games, Inc. All Rights Reserved.
+
+#include "SchemaActions/DataprepFetcherMenuActionCollector.h"
+
+#include "DataprepEditorUtils.h"
+#include "DataprepMenuActionCollectorUtils.h"
+#include "SelectionSystem/DataprepFetcher.h"
+
+#include "ScopedTransaction.h"
+#include "Templates/SubclassOf.h"
+#include "UObject/Class.h"
+
+FDataprepFetcherMenuActionCollector::FDataprepFetcherMenuActionCollector(UDataprepFilter& Filter)
+{
+	FilterPtr = &Filter;
+}
+
+TArray<TSharedPtr<FDataprepSchemaAction>> FDataprepFetcherMenuActionCollector::CollectActions()
+{
+	UDataprepFilter* Filter = FilterPtr.Get();
+	if ( Filter )
+	{
+		UClass* FetcherClass = Filter->GetAcceptedFetcherClass().Get();
+		if ( FetcherClass )
+		{
+			return DataprepMenuActionCollectorUtils::GatherMenuActionForDataprepClass( *FetcherClass
+				, DataprepMenuActionCollectorUtils::FOnCreateMenuAction::CreateRaw( this, &FDataprepFetcherMenuActionCollector::CreateMenuActionFromClass ) );
+		}
+	}
+
+	return {};
+}
+
+bool FDataprepFetcherMenuActionCollector::ShouldAutoExpand()
+{
+	return false;
+}
+
+TSharedPtr<FDataprepSchemaAction> FDataprepFetcherMenuActionCollector::CreateMenuActionFromClass(UClass& Class)
+{
+	UDataprepFilter* Filter = FilterPtr.Get();
+	check( Filter && Class.IsChildOf( Filter->GetAcceptedFetcherClass() ) );
+
+	FDataprepSchemaAction::FOnExecuteAction OnExcuteMenuAction;
+	OnExcuteMenuAction.BindLambda( [ FilterPtr = FilterPtr, Class = &Class] (const FDataprepSchemaActionContext& InContext)
+	{
+		UDataprepFilter* Filter = FilterPtr.Get();
+		if ( Filter && ( !Filter->GetFetcher() || Filter->GetFetcher()->GetClass() != Class ) )
+		{
+			Filter->SetFetcher( TSubclassOf< UDataprepFetcher >( Class ) );
+			FDataprepEditorUtils::NotifySystemOfChangeInPipeline( Filter );
+		}
+	});
+
+	UDataprepFetcher* Fetcher = Class.GetDefaultObject<UDataprepFetcher>();
+	return MakeShared< FDataprepSchemaAction >( FText::FromString( TEXT("") )
+		, Fetcher->GetDisplayFetcherName(), Fetcher->GetTooltipText()
+		, 0, Fetcher->GetAdditionalKeyword(), OnExcuteMenuAction
+		);
+}
diff --git a/Engine/Plugins/Enterprise/DataprepEditor/Source/DataprepEditor/Private/SchemaActions/DataprepFetcherMenuActionCollector.h b/Engine/Plugins/Enterprise/DataprepEditor/Source/DataprepEditor/Private/SchemaActions/DataprepFetcherMenuActionCollector.h
new file mode 100644
--- /dev/null
+++ b/Engine/Plugins/Enterprise/DataprepEditor/Source/DataprepEditor/Private/SchemaActions/DataprepFetcherMenuActionCollector.h
@@ -0,0 +1,30 @@
+// Copyright 1998-2019 Epic Games, Inc. All Rights Reserved.
+
+#pragma once
+
+#include "SchemaActions/IDataprepMenuActionCollector.h"
+#include "SelectionSystem/DataprepFilter.h"
+
+#include "CoreMinimal.h"
+#include "UObject/WeakObjectPtrTemplates.h"
+
+struct FDataprepSchemaAction;
+
+/**
+ * Help collecting the dataprep fetcher menu action for a specified filter
+ * When the action is executed it will set the fetcher on the filter that was pass to this collector
+ */
+class FDataprepFetcherMenuActionCollector : public IDataprepMenuActionCollector
+{
+public:
+	FDataprepFetcherMenuActionCollector(UDataprepFilter& Filter);
+
+	virtual TArray<TSharedPtr<FDataprepSchemaAction>> CollectActions() override;
+	virtual bool ShouldAutoExpand() override;
+
+private:
+
+	TSharedPtr<FDataprepSchemaAction> CreateMenuActionFromClass(UClass& Class);
+
+	TWeakObjectPtr<UDataprepFilter> FilterPtr;
+};
diff --git a/Engine/Plugins/Enterprise/DataprepEditor/Source/DataprepEditor/Private/SchemaActions/DataprepMenuActionCollectorUtils.h b/Engine/Plugins/Enterprise/DataprepEditor/Source/DataprepEditor/Private/SchemaActions/DataprepMenuActionCollectorUtils.h
new file mode 100644
--- /dev/null
+++ b/Engine/Plugins/Enterprise/DataprepEditor/Source/DataprepEditor/Private/SchemaActions/DataprepMenuActionCollectorUtils.h
@@ -0,0 +1,24 @@
+// Copyright 1998-2019 Epic Games, Inc. All Rights Reserved.
+
+#pragma once
+
+#include "SchemaActions/DataprepSchemaAction.h"
+
+#include "CoreMinimal.h"
+#include "UObject/Class.h"
+
+namespace DataprepMenuActionCollectorUtils
+{
+	DECLARE_DELEGATE_RetVal_OneParam(TSharedPtr<FDataprepSchemaAction>, FOnCreateMenuAction,	UClass&);
+
+	/**
+	 * Utils to gather the menu action from a base class 
+	 * @param Class The base class from which we want to create the actions
+	 * @param OnValidClassFound Callback to generate the menu action from the class
+	 */
+	TArray<TSharedPtr<FDataprepSchemaAction>> GatherMenuActionForDataprepClass(UClass& Class, FOnCreateMenuAction OnValidClassFound);
+
+	TArray<UClass*> GetNativeChildClasses(UClass&);
+
+	constexpr EClassFlags NonDesiredClassFlags = CLASS_Deprecated | CLASS_NewerVersionExists | CLASS_Abstract; 
+}
diff --git a/Engine/Plugins/Enterprise/DataprepEditor/Source/DataprepEditor/Private/SchemaActions/DataprepSchemaAction.cpp b/Engine/Plugins/Enterprise/DataprepEditor/Source/DataprepEditor/Private/SchemaActions/DataprepSchemaAction.cpp
new file mode 100644
--- /dev/null
+++ b/Engine/Plugins/Enterprise/DataprepEditor/Source/DataprepEditor/Private/SchemaActions/DataprepSchemaAction.cpp
@@ -0,0 +1,13 @@
+// Copyright 1998-2019 Epic Games, Inc. All Rights Reserved.
+
+#include "SchemaActions/DataprepSchemaAction.h"
+
+FDataprepSchemaAction::FDataprepSchemaAction(FText InActionCategory, FText InMenuDescription, FText InToolTip, const int32 InGrouping, FText InKeywords, const FOnExecuteAction& InAction, int32 InSectionID)
+	: FEdGraphSchemaAction( MoveTemp(InActionCategory), MoveTemp(InMenuDescription), MoveTemp(InToolTip), InGrouping, MoveTemp(InKeywords), InSectionID )
+	, Action( InAction )
+{}
+
+bool FDataprepSchemaAction::ExecuteAction(const FDataprepSchemaActionContext& Context)
+{
+	return Action.ExecuteIfBound( Context );
+}
diff --git a/Engine/Plugins/Enterprise/DataprepEditor/Source/DataprepEditor/Private/Widgets/Action/SDataprepActionBlock.cpp b/Engine/Plugins/Enterprise/DataprepEditor/Source/DataprepEditor/Private/Widgets/Action/SDataprepActionBlock.cpp
new file mode 100644
--- /dev/null
+++ b/Engine/Plugins/Enterprise/DataprepEditor/Source/DataprepEditor/Private/Widgets/Action/SDataprepActionBlock.cpp
@@ -0,0 +1,360 @@
+// Copyright 1998-2019 Epic Games, Inc. All Rights Reserved.
+
+#include "Widgets/Action/SDataprepActionBlock.h"
+
+#include "DataprepActionAsset.h"
+#include "DataprepEditorStyle.h"
+#include "SchemaActions/DataprepDragDropOp.h"
+#include "SchemaActions/DataprepSchemaAction.h"
+
+#include "Editor.h"
+#include "Editor/TransBuffer.h"
+#include "Framework/Application/IMenu.h"
+#include "Framework/Application/SlateApplication.h"
+#include "Framework/Commands/GenericCommands.h"
+#include "Framework/Commands/UICommandList.h"
+#include "Framework/MultiBox/MultiBoxBuilder.h"
+#include "Framework/Text/TextLayout.h"
+#include "Layout/WidgetPath.h"
+#include "ScopedTransaction.h"
+#include "Styling/ISlateStyle.h"
+#include "Styling/SlateStyleRegistry.h"
+#include "Widgets/Colors/SColorBlock.h"
+#include "Widgets/Layout/Anchors.h"
+#include "Widgets/Layout/SConstraintCanvas.h"
+#include "Widgets/SBoxPanel.h"
+#include "Widgets/SNullWidget.h"
+#include "Widgets/Text/STextBlock.h"
+
+#define LOCTEXT_NAMESPACE "SDataprepActionBlock"
+
+void SDataprepActionBlock::Construct(const FArguments& InArgs, const TSharedRef<FDataprepSchemaActionContext>& InDataprepActionContext)
+{
+	DataprepActionContext = InDataprepActionContext;
+	const ISlateStyle* DataprepEditorStyle = FSlateStyleRegistry::FindSlateStyle( FDataprepEditorStyle::GetStyleSetName() );
+	check( DataprepEditorStyle );
+
+	const float DefaultPadding = DataprepEditorStyle->GetFloat( "DataprepAction.Padding" );
+
+	ChildSlot
+	[
+		SNew( SConstraintCanvas )
+
+		// The outline. This is done by a background image
+		+ SConstraintCanvas::Slot()
+		.Anchors( FAnchors( 0.f, 0.f, 1.f, 1.f ) )
+		.Offset( FMargin() )
+		[
+			SNew( SColorBlock )
+			.Color( DataprepEditorStyle->GetColor( "DataprepAction.OutlineColor" ) )
+		]
+
+		+ SConstraintCanvas::Slot()
+		.Anchors( FAnchors( 0.f, 0.f, 1.f, 1.f ) )
+		.Offset( FMargin() )
+		.AutoSize( true )
+		[
+			SNew( SVerticalBox )
+
+			//The title of the block
+			+ SVerticalBox::Slot()
+			.AutoHeight()
+			.Padding( FMargin( DefaultPadding, DefaultPadding, DefaultPadding, 0.f ) )
+			[
+				SNew( SConstraintCanvas )
+
+				// The background of the title
+				+ SConstraintCanvas::Slot()
+				.Anchors( FAnchors( 0.f, 0.f, 1.f, 1.f ) )
+				.Offset( FMargin() )
+				[
+					GetTitleBackgroundWidget()
+				]
+
+				// The title of the block
+				+ SConstraintCanvas::Slot()
+				.Anchors( FAnchors( 0.5f, 0.5f, 0.5f, 0.5f ) )
+				.Offset( FMargin() )
+				.AutoSize( true )
+				[
+					GetTitleWidget()
+				]
+			]
+
+			// The content zone of the action block
+			+ SVerticalBox::Slot()
+			.AutoHeight()
+			.Padding( FMargin( DefaultPadding ) )
+			[
+				SNew( SConstraintCanvas )
+
+				// The background of the content zone
+				+ SConstraintCanvas::Slot()
+				.Anchors( FAnchors( 0.f, 0.f, 1.f, 1.f ) )
+				.Offset( FMargin() )
+				[
+					GetContentBackgroundWidget()
+				]
+
+				// The content of the content zone
+				+ SConstraintCanvas::Slot()
+				.Anchors(FAnchors( 0.f, 0.f, 1.f, 1.f ) )
+				.Offset( FMargin() )
+				.AutoSize( true )
+				[
+					GetContentWidget()
+				]
+			]
+		]
+	];
+}
+
+FReply SDataprepActionBlock::OnMouseButtonDown(const FGeometry& MyGeometry, const FPointerEvent& MouseEvent)
+{
+	if ( MouseEvent.GetEffectingButton() ==  EKeys::LeftMouseButton )
+	{
+		return FReply::Handled().DetectDrag( AsShared(), EKeys::LeftMouseButton );
+	}
+	return FReply::Unhandled();
+}
+
+FReply SDataprepActionBlock::OnMouseButtonUp(const FGeometry& MyGeometry, const FPointerEvent& MouseEvent)
+{
+	if ( MouseEvent.GetEffectingButton() == EKeys::RightMouseButton )
+	{
+		FMenuBuilder MenuBuilder( true, nullptr );
+		PopulateMenuBuilder( MenuBuilder );
+
+		FWidgetPath WidgetPath = MouseEvent.GetEventPath() ? *MouseEvent.GetEventPath() : FWidgetPath();
+		FSlateApplication::Get().PushMenu(
+			AsShared(),
+			WidgetPath,
+			MenuBuilder.MakeWidget(),
+			MouseEvent.GetScreenSpacePosition(),
+			FPopupTransitionEffect(FPopupTransitionEffect::ContextMenu)
+		);
+
+		return FReply::Handled();
+	}
+
+	return FReply::Unhandled();
+}
+
+FReply SDataprepActionBlock::OnDragDetected(const FGeometry& MyGeometry, const FPointerEvent& MouseEvent)
+{
+	check( DataprepActionContext );
+
+	// Is the step moved or copied
+	TSharedRef<bool> IsMoved = MakeShared<bool>( true );
+
+	// Callback that move or copy the operation
+	FDataprepGraphOperation Operation = FDataprepGraphOperation::CreateLambda( [InitialContext = *DataprepActionContext.Get(), IsMoved] (const FDataprepSchemaActionContext& Context)
+		{
+
+			UDataprepActionAsset* InitialAction = InitialContext.DataprepActionPtr.Get();
+			UDataprepActionAsset* DroppedOnAction = Context.DataprepActionPtr.Get();
+			if ( InitialAction && DroppedOnAction )
+			{
+				if ( InitialAction == DroppedOnAction && IsMoved.Get() )
+				{
+					// Move the position of the step in the action
+					if ( InitialContext.StepIndex != INDEX_NONE && InitialContext.StepIndex != Context.StepIndex )
+					{
+						if ( Context.StepIndex == INDEX_NONE )
+						{
+							if ( InitialContext.StepIndex != DroppedOnAction->GetStepsCount() - 1 )
+							{
+								// The step will be move at the end of the array
+								DroppedOnAction->MoveStep( InitialContext.StepIndex, DroppedOnAction->GetStepsCount() - 1 );
+							}
+							else
+							{
+								return false;
+							}
+						}
+						else
+						{
+							DroppedOnAction->MoveStep( InitialContext.StepIndex, Context.StepIndex );
+						}
+						return true;
+					}
+				}
+				else
+				{
+					// Move the step to a another action or duplicate the step
+					UDataprepActionStep* ActionStep = InitialContext.DataprepActionStepPtr.Get();
+					if ( ActionStep )
+					{
+						uint32 NewStepIndex = DroppedOnAction->AddStep( ActionStep );
+						if ( Context.StepIndex != INDEX_NONE )
+						{
+							DroppedOnAction->MoveStep( NewStepIndex, Context.StepIndex );
+						}
+						
+						if ( IsMoved.Get() )
+						{
+							InitialAction->RemoveStep( InitialContext.StepIndex );
+						}
+						return true;
+					}
+				}
+			}
+
+			return false;
+		});
+
+	TSharedRef<FDataprepDragDropOp> DragDropOperation = FDataprepDragDropOp::New( MoveTemp( Operation ) );
+
+	// The pre drop pop the contextual menu allowing the user to select between copy, move and simply cancel the drag and drop.
+	FDataprepPreDropConfirmation PreDropConfirmation = FDataprepPreDropConfirmation::CreateLambda( [ SourceWidget = AsShared(), IsMoved ] (const FDataprepSchemaActionContext& Context, TFunction<void ()> ConfirmationCallback)
+		{
+			auto CopyDraggedItems = [ConfirmationCallback, IsMoved] ()
+				{
+					IsMoved.Get() = false;
+					ConfirmationCallback();
+				};
+
+			auto MoveDraggedItems = [ConfirmationCallback, IsMoved] ()
+				{
+					IsMoved.Get() = true;
+					ConfirmationCallback();
+				};
+
+			FModifierKeysState ModifierKeyState = FSlateApplication::Get().GetModifierKeys();
+
+			bool bCopyKeyDown = ModifierKeyState.IsControlDown()  || ModifierKeyState.IsCommandDown();
+			bool bMoveKeyDown = ModifierKeyState.IsAltDown();
+
+			if ( bCopyKeyDown && !bMoveKeyDown )
+			{
+				// The user is using the shortcut to copy the dragged items
+				IsMoved.Get() = false;
+				ConfirmationCallback();
+			}
+			else if ( bMoveKeyDown && !bCopyKeyDown )
+			{
+				// The user using the shortcut to move the dragged items
+				IsMoved.Get() = true;
+				ConfirmationCallback();
+			}
+			else
+			{
+				FMenuBuilder MenuBuilder(/*bInShouldCloseWindowAfterMenuSelection=*/true, nullptr);
+				const FText MoveCopyHeaderString = LOCTEXT("AssetViewDropMenuHeading", "Move or Copy");
+				MenuBuilder.BeginSection("PathAssetMoveCopy", MoveCopyHeaderString);
+				{
+					MenuBuilder.AddMenuEntry(
+						LOCTEXT("DragDropMove", "Move Here"),
+						LOCTEXT("DragDropMoveTooltip", "Move the dragged items."),
+						FSlateIcon(),
+						FUIAction( FExecuteAction::CreateLambda( MoveDraggedItems ) )
+					);
+					MenuBuilder.AddMenuEntry(
+						LOCTEXT("DragDropCopy", "Copy Here"),
+						LOCTEXT("DragDropCopyTooltip", "Copy the dragged items."),
+						FSlateIcon(),
+						FUIAction( FExecuteAction::CreateLambda( CopyDraggedItems ) )
+					);
+				}
+				MenuBuilder.EndSection();
+
+				TSharedPtr<IMenu> Menu = FSlateApplication::Get().PushMenu(
+					FSlateApplication::Get().GetActiveTopLevelWindow().ToSharedRef(),
+					FWidgetPath(),
+					MenuBuilder.MakeWidget(),
+					FSlateApplication::Get().GetCursorPos(),
+					FPopupTransitionEffect( FPopupTransitionEffect::ContextMenu )
+					);
+
+				if ( GEditor && GEditor->Trans && Menu )
+				{
+					UTransBuffer* TransBuffer = CastChecked<UTransBuffer>( GEditor->Trans );
+					FDelegateHandle Handle = TransBuffer->OnBeforeRedoUndo().AddLambda( [Menu] (const struct FTransactionContext&)
+						{
+							FSlateApplication::Get().DismissMenu( Menu );
+						});
+
+					Menu->GetOnMenuDismissed().AddLambda( [Handle] (TSharedRef<IMenu>)
+						{
+							if ( GEditor && GEditor->Trans )
+							{
+								UTransBuffer* TransBuffer = CastChecked<UTransBuffer>( GEditor->Trans );
+								TransBuffer->OnBeforeRedoUndo().Remove( Handle );
+							}
+						} );
+				}
+			}
+		} );
+
+	DragDropOperation->SetPreDropConfirmation( MoveTemp( PreDropConfirmation ) );
+
+	return FReply::Handled().BeginDragDrop( MoveTemp( DragDropOperation ) );
+}
+
+FText SDataprepActionBlock::GetBlockTitle() const
+{
+	return FText::FromString( TEXT("Default Action Block Title") );
+}
+
+TSharedRef<SWidget> SDataprepActionBlock::GetTitleWidget()
+{
+	const ISlateStyle* DataprepEditorStyle = FSlateStyleRegistry::FindSlateStyle(FDataprepEditorStyle::GetStyleSetName());
+	check( DataprepEditorStyle );
+	const float DefaultPadding = DataprepEditorStyle->GetFloat( "DataprepAction.Padding" );
+
+	return SNew( STextBlock )
+		.Text( GetBlockTitle() )
+		.TextStyle( &DataprepEditorStyle->GetWidgetStyle<FTextBlockStyle>( "DataprepActionBlock.TitleTextBlockStyle" ) )
+		.ColorAndOpacity( FLinearColor( 1.f, 1.f, 1.f ) )
+		.Margin( FMargin( DefaultPadding ) )
+		.Justification( ETextJustify::Center );
+}
+
+TSharedRef<SWidget> SDataprepActionBlock::GetTitleBackgroundWidget()
+{
+	const ISlateStyle* DataprepEditorStyle = FSlateStyleRegistry::FindSlateStyle( FDataprepEditorStyle::GetStyleSetName() );
+	check( DataprepEditorStyle );
+
+	return SNew( SColorBlock )
+		.Color( DataprepEditorStyle->GetColor( "DataprepActionBlock.TitleBackgroundColor" ) );
+}
+
+TSharedRef<SWidget> SDataprepActionBlock::GetContentWidget()
+{
+	return SNullWidget::NullWidget;
+}
+
+TSharedRef<SWidget> SDataprepActionBlock::GetContentBackgroundWidget()
+{
+	const ISlateStyle* DataprepEditorStyle = FSlateStyleRegistry::FindSlateStyle( FDataprepEditorStyle::GetStyleSetName() );
+	check( DataprepEditorStyle );
+
+	return SNew( SColorBlock )
+		.Color( DataprepEditorStyle->GetColor( "DataprepActionBlock.ContentBackgroundColor" ) );
+}
+
+void SDataprepActionBlock::PopulateMenuBuilder(FMenuBuilder& MenuBuilder)
+{
+	FUIAction DeleteAction;
+	DeleteAction.ExecuteAction.BindSP( this, &SDataprepActionBlock::DeleteStep );
+
+	TSharedPtr<FUICommandInfo> DeleteCommand = FGenericCommands::Get().Delete;
+	MenuBuilder.AddMenuEntry( DeleteCommand->GetLabel(),
+		DeleteCommand->GetDescription(),
+		DeleteCommand->GetIcon(),
+		DeleteAction );
+}
+
+void SDataprepActionBlock::DeleteStep()
+{
+	if ( FDataprepSchemaActionContext* ActionContext = DataprepActionContext.Get() )
+	{
+		if ( UDataprepActionAsset* ActionAsset = ActionContext->DataprepActionPtr.Get() )
+		{
+			FScopedTransaction Transaction( LOCTEXT("DeleteStepTransaction", "Remove step from action") );
+			ActionAsset->RemoveStep( ActionContext->StepIndex );
+		}
+	}
+}
+
+#undef LOCTEXT_NAMESPACE
diff --git a/Engine/Plugins/Enterprise/DataprepEditor/Source/DataprepEditor/Private/Widgets/Action/SDataprepActionBlock.h b/Engine/Plugins/Enterprise/DataprepEditor/Source/DataprepEditor/Private/Widgets/Action/SDataprepActionBlock.h
new file mode 100644
--- /dev/null
+++ b/Engine/Plugins/Enterprise/DataprepEditor/Source/DataprepEditor/Private/Widgets/Action/SDataprepActionBlock.h
@@ -0,0 +1,70 @@
+// Copyright 1998-2019 Epic Games, Inc. All Rights Reserved.
+
+#pragma once
+
+#include "CoreMinimal.h"
+#include "Widgets/SCompoundWidget.h"
+#include "Widgets/DeclarativeSyntaxSupport.h"
+
+struct FDataprepSchemaActionContext;
+struct FTextBlockStyle;
+
+/**
+ * The SDataprepActionBlock is the base class that is used for the block that compose a dataprep action
+ */
+class SDataprepActionBlock : public SCompoundWidget
+{
+	SLATE_BEGIN_ARGS(SDataprepActionBlock) {}
+	SLATE_END_ARGS();
+
+	void Construct(const FArguments& InArgs, const TSharedRef<FDataprepSchemaActionContext>& DataprepActionContext);
+
+	// SWidget Interface
+	virtual FReply OnMouseButtonDown(const FGeometry& MyGeometry, const FPointerEvent& MouseEvent) override;
+	virtual FReply OnMouseButtonUp(const FGeometry& MyGeometry, const FPointerEvent& MouseEvent) override;
+	virtual FReply OnDragDetected(const FGeometry& MyGeometry, const FPointerEvent& MouseEvent) override;
+	// End of SWidget Interface
+
+protected:
+
+	/**
+	 * Return the block title.
+	 * This is used in the implementation of SDataprepActionBlock::GetTitleWidget
+	 */
+	virtual FText GetBlockTitle() const;
+
+	/**
+	 * Return the title widget
+	 */
+ 	virtual TSharedRef<SWidget> GetTitleWidget();
+
+	/**
+	 * Get title background widget
+	 */
+	virtual TSharedRef<SWidget> GetTitleBackgroundWidget();
+
+	/**
+	 * Return the content widget
+	 */
+	virtual TSharedRef<SWidget> GetContentWidget();
+
+	/**
+	 * Get content background widget
+	 */
+	virtual TSharedRef<SWidget> GetContentBackgroundWidget();
+
+	/**
+	 * Populate the right click menu by overriding this function.
+	 * Call the parent implementation if you want the full menu.
+	 */
+	virtual void PopulateMenuBuilder(class FMenuBuilder& MenuBuilder);
+
+private:
+
+	/**
+	 * Delete this step
+	 */
+	void DeleteStep();
+
+	TSharedPtr<FDataprepSchemaActionContext> DataprepActionContext;
+};
diff --git a/Engine/Plugins/Enterprise/DataprepEditor/Source/DataprepEditor/Private/Widgets/Action/SDataprepBoolFilter.cpp b/Engine/Plugins/Enterprise/DataprepEditor/Source/DataprepEditor/Private/Widgets/Action/SDataprepBoolFilter.cpp
new file mode 100644
--- /dev/null
+++ b/Engine/Plugins/Enterprise/DataprepEditor/Source/DataprepEditor/Private/Widgets/Action/SDataprepBoolFilter.cpp
@@ -0,0 +1,12 @@
+// Copyright 1998-2019 Epic Games, Inc. All Rights Reserved.
+
+#include "Widgets/Action/SDataprepBoolFilter.h"
+
+#include "SelectionSystem/DataprepBoolFilter.h"
+#include "Widgets/Action/SDataprepFetcherSelector.h"
+
+#include "Widgets/Layout/SBox.h"
+
+void SDataprepBoolFilter::Construct(const FArguments& InArgs, UDataprepBoolFilter& InFilter)
+{
+}
diff --git a/Engine/Plugins/Enterprise/DataprepEditor/Source/DataprepEditor/Private/Widgets/Action/SDataprepBoolFilter.h b/Engine/Plugins/Enterprise/DataprepEditor/Source/DataprepEditor/Private/Widgets/Action/SDataprepBoolFilter.h
new file mode 100644
--- /dev/null
+++ b/Engine/Plugins/Enterprise/DataprepEditor/Source/DataprepEditor/Private/Widgets/Action/SDataprepBoolFilter.h
@@ -0,0 +1,17 @@
+// Copyright 1998-2019 Epic Games, Inc. All Rights Reserved.
+
+#pragma once
+
+#include "CoreMinimal.h"
+#include "Widgets/DeclarativeSyntaxSupport.h"
+#include "Widgets/SCompoundWidget.h"
+
+class UDataprepBoolFilter;
+
+class SDataprepBoolFilter : public SCompoundWidget
+{
+	SLATE_BEGIN_ARGS(SDataprepBoolFilter) {}
+	SLATE_END_ARGS();
+
+	void Construct(const FArguments& InArgs, UDataprepBoolFilter& InFilter);
+};
diff --git a/Engine/Plugins/Enterprise/DataprepEditor/Source/DataprepEditor/Private/Widgets/Action/SDataprepFetcherSelector.cpp b/Engine/Plugins/Enterprise/DataprepEditor/Source/DataprepEditor/Private/Widgets/Action/SDataprepFetcherSelector.cpp
new file mode 100644
--- /dev/null
+++ b/Engine/Plugins/Enterprise/DataprepEditor/Source/DataprepEditor/Private/Widgets/Action/SDataprepFetcherSelector.cpp
@@ -0,0 +1,53 @@
+// Copyright 1998-2019 Epic Games, Inc. All Rights Reserved.
+
+#include "Widgets/Action/SDataprepFetcherSelector.h"
+
+#include "SchemaActions/DataprepFetcherMenuActionCollector.h"
+#include "SelectionSystem/DataprepFilter.h"
+#include "Widgets/SDataprepActionMenu.h"
+
+#include "Framework/Text/TextLayout.h"
+#include "Widgets/Input/SComboButton.h"
+#include "Widgets/Input/SEditableTextBox.h"
+
+void SDataprepFetcherSelector::Construct(const FArguments& InArgs, UDataprepFilter& InFilter)
+{
+	Filter = &InFilter;
+	
+	ChildSlot
+	[
+		SAssignNew( FetcherTypeButton, SComboButton )
+		.OnGetMenuContent( this, &SDataprepFetcherSelector::GetFetcherTypeSelector )
+		.ButtonContent()
+		[
+			SNew( STextBlock )
+			.Text( this, &SDataprepFetcherSelector::GetFetcherNameText )
+			.Justification( ETextJustify::Center )
+		]
+	];
+}
+
+FText SDataprepFetcherSelector::GetFetcherNameText() const
+{
+	check( Filter );
+	UDataprepFetcher* Fetcher = Filter->GetFetcher();
+	if ( Fetcher )
+	{
+		return Fetcher->GetDisplayFetcherName();
+	}
+
+	return {};
+}
+
+TSharedRef<SWidget> SDataprepFetcherSelector::GetFetcherTypeSelector() const
+{
+	TSharedRef< SDataprepActionMenu > ActionMenu = SNew( SDataprepActionMenu, MakeUnique< FDataprepFetcherMenuActionCollector >( *Filter ) )
+		.TransactionText( NSLOCTEXT( "SDataprepFetcherSelector", "ChangingFetcher", "Change Fetcher Type" ) );
+	FetcherTypeButton->SetMenuContentWidgetToFocus( ActionMenu->GetFilterTextBox() );
+	return ActionMenu;
+}
+
+void SDataprepFetcherSelector::AddReferencedObjects(FReferenceCollector& Collector)
+{
+	Collector.AddReferencedObject( Filter );
+}
diff --git a/Engine/Plugins/Enterprise/DataprepEditor/Source/DataprepEditor/Private/Widgets/Action/SDataprepFetcherSelector.h b/Engine/Plugins/Enterprise/DataprepEditor/Source/DataprepEditor/Private/Widgets/Action/SDataprepFetcherSelector.h
new file mode 100644
--- /dev/null
+++ b/Engine/Plugins/Enterprise/DataprepEditor/Source/DataprepEditor/Private/Widgets/Action/SDataprepFetcherSelector.h
@@ -0,0 +1,34 @@
+// Copyright 1998-2019 Epic Games, Inc. All Rights Reserved.
+
+#pragma once
+
+#include "CoreMinimal.h"
+#include "UObject/GCObject.h"
+#include "Widgets/DeclarativeSyntaxSupport.h"
+#include "Widgets/SCompoundWidget.h"
+
+class SComboButton;
+class UDataprepFilter;
+
+/**
+ * The dataprep fetcher selector is widget made to change and display the current type of fetcher of a dataprep filter
+ */
+class SDataprepFetcherSelector : public SCompoundWidget, public FGCObject
+{
+	SLATE_BEGIN_ARGS(SDataprepFetcherSelector) {}
+	SLATE_END_ARGS();
+
+	void Construct(const FArguments& InArgs, UDataprepFilter& InFilter);
+
+private:
+	// Those function are for the fetcher display
+	TSharedRef<SWidget> GetFetcherTypeSelector() const;
+	FText GetFetcherNameText() const;
+
+	//~ FGCObject interface
+	virtual void AddReferencedObjects(FReferenceCollector& Collector) override;
+	//~ End FGCObject interface
+
+	UDataprepFilter* Filter;
+	TSharedPtr<SComboButton> FetcherTypeButton;
+};
diff --git a/Engine/Plugins/Enterprise/DataprepEditor/Source/DataprepEditor/Private/Widgets/Action/SDataprepOperation.h b/Engine/Plugins/Enterprise/DataprepEditor/Source/DataprepEditor/Private/Widgets/Action/SDataprepOperation.h
new file mode 100644
--- /dev/null
+++ b/Engine/Plugins/Enterprise/DataprepEditor/Source/DataprepEditor/Private/Widgets/Action/SDataprepOperation.h
@@ -0,0 +1,35 @@
+// Copyright 1998-2019 Epic Games, Inc. All Rights Reserved.
+
+#pragma once
+
+#include "Widgets/Action/SDataprepActionBlock.h"
+
+#include "CoreMinimal.h"
+#include "UObject/GCObject.h"
+#include "Widgets/DeclarativeSyntaxSupport.h"
+#include "Widgets/SCompoundWidget.h"
+
+class UDataprepOperation;
+struct FDataprepSchemaActionContext;
+
+class SDataprepOperation : public SDataprepActionBlock, public FGCObject
+{
+	SLATE_BEGIN_ARGS(SDataprepOperation) {}
+	SLATE_END_ARGS();
+
+	void Construct(const FArguments& InArgs, UDataprepOperation* InOperation, const TSharedRef<FDataprepSchemaActionContext>& InDataprepActionContext);
+
+protected:
+
+	virtual FText GetBlockTitle() const override;
+	virtual TSharedRef<SWidget> GetContentWidget() override;
+
+private:
+
+	FText GetTooltipText() const;
+
+	//~ FGCObject interface
+	virtual void AddReferencedObjects(FReferenceCollector& Collector) override;
+	 
+	UDataprepOperation* Operation;
+};
\ No newline at end of file
diff --git a/Engine/Plugins/Enterprise/DataprepEditor/Source/DataprepEditor/Private/Widgets/Action/SGraphNodeK2DataprepAction.h b/Engine/Plugins/Enterprise/DataprepEditor/Source/DataprepEditor/Private/Widgets/Action/SGraphNodeK2DataprepAction.h
new file mode 100644
--- /dev/null
+++ b/Engine/Plugins/Enterprise/DataprepEditor/Source/DataprepEditor/Private/Widgets/Action/SGraphNodeK2DataprepAction.h
@@ -0,0 +1,24 @@
+// Copyright 1998-2019 Epic Games, Inc. All Rights Reserved.
+
+#pragma once
+
+#include "CoreMinimal.h"
+#include "KismetNodes/SGraphNodeK2Default.h"
+#include "Widgets/DeclarativeSyntaxSupport.h"
+
+class UK2Node_DataprepAction;
+
+class SGraphNodeK2DataprepAction : public SGraphNodeK2Default
+{
+public:
+	SLATE_BEGIN_ARGS(SGraphNodeK2DataprepAction) {}
+	SLATE_END_ARGS()
+
+	void Construct(const FArguments& InArgs, UK2Node_DataprepAction* InActionNode);
+
+	//  Add the widgets below the pins 
+	virtual void CreateBelowPinControls(TSharedPtr<SVerticalBox> MainBox) override;
+
+private:
+	TWeakObjectPtr<class UDataprepActionAsset> DataprepActionPtr;
+};
diff --git a/Engine/Plugins/Enterprise/DataprepEditor/Source/DataprepEditor/Private/Widgets/Action/SGraphNodek2DataprepAction.cpp b/Engine/Plugins/Enterprise/DataprepEditor/Source/DataprepEditor/Private/Widgets/Action/SGraphNodek2DataprepAction.cpp
new file mode 100644
--- /dev/null
+++ b/Engine/Plugins/Enterprise/DataprepEditor/Source/DataprepEditor/Private/Widgets/Action/SGraphNodek2DataprepAction.cpp
@@ -0,0 +1,43 @@
+// Copyright 1998-2019 Epic Games, Inc. All Rights Reserved.
+
+#include "Widgets/Action/SGraphNodeK2DataprepAction.h"
+
+// Dataprep includes
+#include "BlueprintNodes/K2Node_DataprepAction.h"
+#include "DataprepActionAsset.h"
+#include "Widgets/Action/SDataprepActionSteps.h"
+
+// Engine Includes
+#include "Widgets/SBoxPanel.h"
+
+void SGraphNodeK2DataprepAction::Construct(const FArguments& InArgs, UK2Node_DataprepAction* InActionNode)
+{
+	DataprepActionPtr = InActionNode->GetDataprepAction();
+
+	this->GraphNode = InActionNode;
+	this->SetCursor(EMouseCursor::CardinalCross);
+	this->UpdateGraphNode();
+}
+
+void SGraphNodeK2DataprepAction::CreateBelowPinControls(TSharedPtr<SVerticalBox> MainBox)
+{
+	UDataprepActionAsset* DataprepAction = DataprepActionPtr.Get();
+	if ( DataprepAction )
+	{
+		MainBox->AddSlot()
+		.AutoHeight()
+		[
+			SNew( SDataprepActionSteps, DataprepAction )
+		];
+	}
+	else
+	{
+		MainBox->AddSlot()
+			.AutoHeight()
+			[
+				SNew(STextBlock)
+				.ColorAndOpacity( FSlateColor( FLinearColor::Red ) )
+				.Text( FText::FromString( TEXT("This node doesn't have a dataprep action!") ) )
+			];
+	}
+}
diff --git a/Engine/Plugins/Enterprise/DataprepEditor/Source/DataprepEditor/Private/Widgets/DataprepAssetView.h b/Engine/Plugins/Enterprise/DataprepEditor/Source/DataprepEditor/Private/Widgets/DataprepAssetView.h
new file mode 100644
--- /dev/null
+++ b/Engine/Plugins/Enterprise/DataprepEditor/Source/DataprepEditor/Private/Widgets/DataprepAssetView.h
@@ -0,0 +1,118 @@
+// Copyright 1998-2019 Epic Games, Inc. All Rights Reserved.
+
+#pragma once
+
+#include "CoreMinimal.h"
+
+#include "DataprepAssetInterface.h"
+#include "DataprepAssetProducers.h"
+#include "DataprepContentProducer.h"
+#include "DataprepEditor.h"
+#include "DataprepWidgets.h"
+
+#include "Delegates/IDelegateInstance.h"
+#include "Widgets/Input/STextComboBox.h"
+#include "Widgets/SCompoundWidget.h"
+#include "Widgets/Text/STextBlock.h"
+#include "Widgets/Views/STableRow.h"
+#include "Widgets/Views/STreeView.h"
+
+class IDetailsView;
+class SDataprepAssetView;
+class SDataprepConsumerWidget;
+class SDataprepProducersWidget;
+
+class SDataprepAssetView : public SCompoundWidget
+{
+public:
+	SLATE_BEGIN_ARGS(SDataprepAssetView) {}
+	SLATE_END_ARGS()
+
+	void Construct(const FArguments& InArgs, UDataprepAssetInterface* InDataprepAssetPtr, TSharedPtr<FUICommandList>& InCommandList);
+
+	~SDataprepAssetView();
+
+private:
+	void OnNewConsumerSelected( TSharedPtr<FString> NewConsumer, ESelectInfo::Type SelectInfo);
+
+	/** Handles changes in the Dataprep asset */
+	void OnDataprepAssetChanged(FDataprepAssetChangeType ChangeType);
+
+	float OnGetLeftColumnWidth() const { return 1.0f - ColumnWidth; }
+	float OnGetRightColumnWidth() const { return ColumnWidth; }
+	void OnSetColumnWidth(float InWidth) { ColumnWidth = InWidth; }
+
+	TSharedRef<ITableRow> OnGenerateRowForCategoryTree( TSharedRef<EDataprepCategory> InTreeNode, const TSharedRef<STableViewBase>& InOwnerTable );
+	void OnGetChildrenForCategoryTree( TSharedRef<EDataprepCategory> InTreeNode, TArray< TSharedRef<EDataprepCategory> >& OutChildren ) {}
+
+	TArray<TSharedRef< EDataprepCategory >> Categories;
+	TSharedPtr<FUICommandList> CommandList;
+
+private:
+	TWeakObjectPtr<UDataprepAssetInterface> DataprepAssetInterfacePtr;
+	TSharedPtr<SDataprepProducersWidget> ProducersWidget;
+	TSharedPtr< STextBlock > CheckBox;
+	TArray< TSharedPtr< FString > > ConsumerDescriptionList;
+	TMap< TSharedPtr< FString >, UClass* > ConsumerDescriptionMap;
+	TSharedPtr< FString > SelectedConsumerDescription;
+	TSharedPtr< SWidget > ConsumerSelector;
+	bool bIsChecked;
+	TSharedPtr< SDataprepConsumerWidget > ConsumerWidget;
+
+	/** Container used by all splitters in the details view, so that they move in sync */
+	TSharedPtr< FDataprepDetailsViewColumnSizeData > ColumnSizeData;
+
+	/** Relative width to control splitters */
+	float ColumnWidth;
+
+	FDelegateHandle OnParameterizationWasEdited;
+};
+
+// Inspired from SKismetInspector class
+class SGraphNodeDetailsWidget : public SCompoundWidget
+{
+public:
+	SLATE_BEGIN_ARGS(SGraphNodeDetailsWidget) {}
+	SLATE_END_ARGS()
+
+public:
+	void Construct(const FArguments& InArgs);
+
+	void ShowDetailsObjects(const TArray<UObject*>& Objects);
+
+	void SetCanEditProperties(bool bInCanEditProperties) { bCanEditProperties = bInCanEditProperties; };
+	bool GetCanEditProperties() const { return bCanEditProperties; };
+	const TArray< TWeakObjectPtr<UObject> >& GetObjectsShowInDetails() const { return SelectedObjects; };
+
+	// SWidget interface
+	virtual void Tick(const FGeometry& AllottedGeometry, const double InCurrentTime, const float InDeltaTime) override;
+	// End of SWidget interface
+
+private:
+	/** Update the inspector window to show information on the supplied objects */
+	void UpdateFromObjects(const TArray<UObject*>& PropertyObjects);
+
+	/** Add this property and all its child properties to SelectedObjectProperties */
+	void AddPropertiesRecursive(UProperty* Property);
+
+private:
+	/** Property viewing widget */
+	TSharedPtr<IDetailsView> PropertyView;
+
+	/** Border widget that wraps a dynamic context-sensitive widget for editing objects that the property window is displaying */
+	TSharedPtr<SBorder> ContextualEditingBorderWidget;
+
+	/** Selected objects for this detail view */
+	TArray< TWeakObjectPtr<UObject> > SelectedObjects;
+
+	/** Set of object properties that should be visible */
+	TSet<TWeakObjectPtr<UProperty> > SelectedObjectProperties;
+
+	/** When TRUE, the SGraphNodeDetailsWidget needs to refresh the details view on Tick */
+	bool bRefreshOnTick;
+
+	bool bCanEditProperties;
+
+	/** Holds the property objects that need to be displayed by the inspector starting on the next tick */
+	TArray<UObject*> RefreshPropertyObjects;
+};
