Trying to fix the case issue on linux

diff --git a/Engine/Plugins/Enterprise/DataprepEditor/Source/DataprepEditor/Private/Widgets/DataprepWidgets.cpp b/Engine/Plugins/Enterprise/DataprepEditor/Source/DataprepEditor/Private/Widgets/DataprepWidgets.cpp
new file mode 100644
--- /dev/null
+++ b/Engine/Plugins/Enterprise/DataprepEditor/Source/DataprepEditor/Private/Widgets/DataprepWidgets.cpp
@@ -0,0 +1,1224 @@
+// Copyright 1998-2019 Epic Games, Inc. All Rights Reserved.
+
+#include "DataprepWidgets.h"
+
+#include "DataprepAsset.h"
+#include "DataprepAssetInstance.h"
+#include "DataprepContentConsumer.h"
+#include "DataprepAssetView.h"
+#include "DataprepEditorLogCategory.h"
+#include "DataprepEditorUtils.h"
+#include "DataprepParameterizableObject.h"
+#include "Parameterization/DataprepParameterizationUtils.h"
+#include "Widgets/Parameterization/SDataprepParameterizationLinkIcon.h"
+
+#include "ContentBrowserModule.h"
+#include "DetailLayoutBuilder.h"
+#include "Dialogs/DlgPickPath.h"
+#include "Editor.h"
+#include "EditorFontGlyphs.h"
+#include "EditorFontGlyphs.h"
+#include "EditorStyleSet.h"
+#include "Framework/Application/MenuStack.h"
+#include "Framework/Application/SlateApplication.h"
+#include "Framework/MultiBox/MultiBoxBuilder.h"
+#include "Framework/MultiBox/MultiBoxBuilder.h"
+#include "HAL/PlatformProcess.h"
+#include "IContentBrowserSingleton.h"
+#include "IDetailTreeNode.h"
+#include "IPropertyRowGenerator.h"
+#include "Input/Events.h"
+#include "Layout/Geometry.h"
+#include "Layout/WidgetPath.h"
+#include "Misc/Paths.h"
+#include "Modules/ModuleManager.h"
+#include "PropertyCustomizationHelpers.h"
+#include "PropertyEditorModule.h"
+#include "PropertyEditorModule.h"
+#include "ScopedTransaction.h"
+#include "Styling/CoreStyle.h"
+#include "UObject/StrongObjectPtr.h"
+#include "UObject/UnrealType.h"
+#include "Widgets/Images/SImage.h"
+#include "Widgets/Input/SButton.h"
+#include "Widgets/Input/SComboButton.h"
+#include "Widgets/Input/SEditableTextBox.h"
+#include "Widgets/Layout/SBox.h"
+#include "Widgets/Layout/SGridPanel.h"
+#include "Widgets/Layout/SSpacer.h"
+#include "Widgets/SNullWidget.h"
+#include "Widgets/Text/SInlineEditableTextBlock.h"
+#include "Widgets/Text/STextBlock.h"
+#include "UObject/UObjectGlobals.h"
+
+#define LOCTEXT_NAMESPACE "DataprepSlateHelper"
+
+namespace DataprepWidgetUtils
+{
+	TSharedRef<SWidget> CreatePropertyWidget( TSharedPtr<SWidget> NameWidget, TSharedPtr<SWidget> ValueWidget, TSharedPtr< FDataprepDetailsViewColumnSizeData > ColumnSizeData, float Spacing, EHorizontalAlignment HAlign = EHorizontalAlignment::HAlign_Left, EVerticalAlignment VAlign = EVerticalAlignment::VAlign_Center )
+	{
+		return SNew(SHorizontalBox)
+		+ SHorizontalBox::Slot()
+		.FillWidth(1.f)
+		.VAlign(VAlign_Center)
+		.HAlign(HAlign_Left)
+		.Padding(0.0f, 0.0f, 0.0f, Spacing)
+		[
+			SNew(SSplitter)
+			.Style(FEditorStyle::Get(), "DetailsView.Splitter")
+			.PhysicalSplitterHandleSize(1.0f)
+			.HitDetectionSplitterHandleSize(5.0f)
+			.ResizeMode( ESplitterResizeMode::Fill )
+			+ SSplitter::Slot()
+			.Value(ColumnSizeData->LeftColumnWidth)
+			.OnSlotResized( SSplitter::FOnSlotResized::CreateLambda( [](float InNewWidth) -> void {} ) )
+			[
+				NameWidget.ToSharedRef()
+			]
+			+ SSplitter::Slot()
+			.Value(ColumnSizeData->RightColumnWidth)
+			.OnSlotResized(ColumnSizeData->OnWidthChanged)
+			[
+				SNew(SHorizontalBox)
+				.Clipping(EWidgetClipping::OnDemand)
+				+ SHorizontalBox::Slot()
+				.Padding( 5.0f, 2.5f, 2.0f, 2.5f )
+				.HAlign( HAlign )
+				.VAlign( VAlign )
+				[
+					// Trick to force the splitter widget to fill up the space of its parent
+					// Strongly inspired from SDetailSingleItemRow
+					SNew(DataprepWidgetUtils::SConstrainedBox)
+					[
+						ValueWidget.ToSharedRef()
+					]
+				]
+			]
+		];
+	}
+}
+
+void SDataprepCategoryWidget::ToggleExpansion()
+{
+	bIsExpanded = !bIsExpanded;
+	CategoryContent->SetVisibility(bIsExpanded ? EVisibility::Visible : EVisibility::Collapsed);
+}
+
+const FSlateBrush* SDataprepCategoryWidget::GetBackgroundImage() const
+{
+	if (IsHovered())
+	{
+		return bIsExpanded ? FEditorStyle::GetBrush("DetailsView.CategoryTop_Hovered") : FEditorStyle::GetBrush("DetailsView.CollapsedCategory_Hovered");
+	} 
+	else
+	{
+		return bIsExpanded ? FEditorStyle::GetBrush("DetailsView.CategoryTop") : FEditorStyle::GetBrush("DetailsView.CollapsedCategory");
+	}
+}
+
+void SDataprepCategoryWidget::Construct( const FArguments& InArgs, TSharedRef< SWidget > InContent, const TSharedRef<STableViewBase>& InOwnerTableView )
+{
+	const float MyContentTopPadding = 2.0f;
+	const float MyContentBottomPadding = 2.0f;
+
+	const float ChildSlotPadding = 2.0f;
+	const float BorderVerticalPadding = 3.0f;
+
+	CategoryContent = InContent;
+
+	TSharedPtr<SWidget> TitleDetail = InArgs._TitleDetail;
+	if( !TitleDetail.IsValid() )
+	{
+		TitleDetail = SNullWidget::NullWidget;
+	}
+
+	TSharedPtr<SHorizontalBox> TitleHeader = SNew(SHorizontalBox)
+		+ SHorizontalBox::Slot()
+		.VAlign( VAlign_Center )
+		.Padding( 2.0f, MyContentTopPadding, 2.0f, MyContentBottomPadding )
+		.AutoWidth()
+		[
+			SNew( SExpanderArrow, SharedThis(this) )
+		]
+		+ SHorizontalBox::Slot()
+		.VAlign( VAlign_Center )
+		.AutoWidth()
+		.Padding( 0.0f, 8.0f )
+		[
+			SNew(STextBlock)
+			.Text( InArgs._Title )
+			.Font( FEditorStyle::GetFontStyle("DetailsView.CategoryFontStyle") )
+			.ShadowOffset( FVector2D( 1.0f, 1.0f ) )
+		];
+
+	ChildSlot
+	[
+		SNew(SVerticalBox)
+		+ SVerticalBox::Slot()
+		.AutoHeight()
+		[
+			SNew( SBorder )
+			.BorderImage( this, &SDataprepCategoryWidget::GetBackgroundImage )
+			.Padding( FMargin( 0.0f, BorderVerticalPadding, 16.0f, BorderVerticalPadding ) )
+			.BorderBackgroundColor( FLinearColor( .6, .6, .6, 1.0f ) )
+			[
+				SNew(SHorizontalBox)
+				+ SHorizontalBox::Slot()
+				.FillWidth( 0.5f )
+				.HAlign( EHorizontalAlignment::HAlign_Left )
+				[
+					TitleHeader.ToSharedRef()
+				]
+				+ SHorizontalBox::Slot()
+				.FillWidth( 0.5f )
+				.HAlign( EHorizontalAlignment::HAlign_Right )
+				[
+					TitleDetail.ToSharedRef()
+				]
+			]
+		]
+		+ SVerticalBox::Slot()
+		.AutoHeight()
+		[
+			CategoryContent.ToSharedRef()
+		]
+	];
+
+	STableRow< TSharedPtr< EDataprepCategory > >::ConstructInternal(
+		STableRow::FArguments()
+		.Style( FEditorStyle::Get(), "DetailsView.TreeView.TableRow" )
+		.ShowSelection( false ),
+		InOwnerTableView
+	);
+}
+
+void SDataprepConsumerWidget::OnLevelNameChanged( const FText &NewLevelName, ETextCommit::Type CommitType )
+{
+	const FScopedTransaction Transaction( LOCTEXT("Consumer_SetLevelName", "Set Level Name") );
+	FText OutReason;
+	if( !DataprepConsumerPtr->SetLevelName( NewLevelName.ToString(), OutReason ) )
+	{
+		UE_LOG( LogDataprepEditor, Error, TEXT("%s"), *OutReason.ToString() );
+		LevelTextBox->SetText( FText::FromString( DataprepConsumerPtr->GetLevelName() ) );
+	}
+}
+
+void SDataprepConsumerWidget::OnTextCommitted( const FText& NewText, ETextCommit::Type CommitType)
+{
+	if( UDataprepContentConsumer* DataprepConsumer = DataprepConsumerPtr.Get() )
+	{
+		FString NewContentFolder( NewText.ToString() );
+
+		// Replace /Content/ with /Game/ since /Content is only used for display 
+		if( NewContentFolder.StartsWith( TEXT("/Content") ) )
+		{
+			NewContentFolder = NewContentFolder.Replace( TEXT( "/Content" ), TEXT( "/Game" ) );
+		}
+		
+		// Remove ending '/' if applicable
+		if( !NewContentFolder.IsEmpty() && NewContentFolder[ NewContentFolder.Len()-1 ] == TEXT('/') )
+		{
+			NewContentFolder[ NewContentFolder.Len()-1 ] = 0;
+			NewContentFolder = NewContentFolder.LeftChop(1);
+		}
+
+		const FScopedTransaction Transaction( LOCTEXT("Consumer_SetTargetContentFolder", "Set Target Content Folder") );
+
+		FText ErrorReason;
+		if( !DataprepConsumer->SetTargetContentFolder( NewContentFolder, ErrorReason ) )
+		{
+			UE_LOG( LogDataprepEditor, Error, TEXT("%s"), *ErrorReason.ToString() );
+			UpdateContentFolderText();
+		}
+	}
+}
+
+void SDataprepConsumerWidget::OnBrowseContentFolder()
+{
+	FString Path = DataprepConsumerPtr->GetTargetContentFolder();
+	if( Path.IsEmpty() )
+	{
+		Path = FPaths::GetPath( DataprepConsumerPtr->GetOutermost()->GetPathName() );
+	}
+	Path += TEXT("/"); // Trailing '/' is needed to set the default path
+
+	//Ask the user for the root path where they want any content to be placed
+	if( UDataprepContentConsumer* DataprepConsumer = DataprepConsumerPtr.Get() )
+	{
+		TSharedRef<SDlgPickPath> PickContentPathDlg =
+			SNew(SDlgPickPath)
+			.Title(LOCTEXT("DataprepSlateHelper_ChooseImportRootContentPath", "Choose Location for importing the Datasmith content"))
+			.DefaultPath(FText::FromString(Path));
+
+		if ( PickContentPathDlg->ShowModal() == EAppReturnType::Ok )
+		{
+			const FScopedTransaction Transaction( LOCTEXT("Consumer_SetTargetContentFolder", "Set Target Content Folder") );
+
+			FText ErrorReason;
+			if( DataprepConsumer->SetTargetContentFolder( PickContentPathDlg->GetPath().ToString(), ErrorReason ) )
+			{
+				UpdateContentFolderText();
+			}
+			else
+			{
+				UE_LOG( LogDataprepEditor, Error, TEXT("%s"), *ErrorReason.ToString() );
+			}
+		}
+	}
+}
+
+void SDataprepConsumerWidget::SetDataprepConsumer(UDataprepContentConsumer* InDataprepConsumer)
+{
+	if(InDataprepConsumer == nullptr)
+	{
+		return;
+	}
+
+	if(DataprepConsumerPtr != InDataprepConsumer)
+	{
+		if(DataprepConsumerPtr.IsValid())
+		{
+			DataprepConsumerPtr->GetOnChanged().Remove( OnConsumerChangedHandle );
+		}
+
+		DataprepConsumerPtr = InDataprepConsumer;
+
+		OnConsumerChangedHandle = InDataprepConsumer->GetOnChanged().AddSP( this, &SDataprepConsumerWidget::OnConsumerChanged );
+
+		OnConsumerChanged();
+	}
+}
+
+void SDataprepConsumerWidget::Construct(const FArguments& InArgs )
+{
+	if (InArgs._ColumnSizeData.IsValid())
+	{
+		ColumnSizeData = InArgs._ColumnSizeData;
+	}
+	else
+	{
+		ColumnWidth = 0.7f;
+		ColumnSizeData = MakeShared<FDataprepDetailsViewColumnSizeData>();
+		ColumnSizeData->LeftColumnWidth = TAttribute<float>(this, &SDataprepConsumerWidget::OnGetLeftColumnWidth);
+		ColumnSizeData->RightColumnWidth = TAttribute<float>(this, &SDataprepConsumerWidget::OnGetRightColumnWidth);
+		ColumnSizeData->OnWidthChanged = SSplitter::FOnSlotResized::CreateSP(this, &SDataprepConsumerWidget::OnSetColumnWidth);
+	}
+
+	TSharedPtr<SWidget> InnerWidget = InArgs._DataprepConsumer ? BuildWidget() : BuildNullWidget();
+
+	ChildSlot
+	[
+		InnerWidget.ToSharedRef()
+	];
+
+	if(InArgs._DataprepConsumer)
+	{
+		SetDataprepConsumer( InArgs._DataprepConsumer );
+	}
+}
+
+SDataprepConsumerWidget::~SDataprepConsumerWidget()
+{
+	if ( UDataprepContentConsumer* DataprepConsumer = DataprepConsumerPtr.Get() )
+	{
+		DataprepConsumer->GetOnChanged().Remove( OnConsumerChangedHandle );
+	}
+}
+
+TSharedRef<SWidget> SDataprepConsumerWidget::BuildWidget()
+{
+	TSharedRef<SWidget> BrowseButton = PropertyCustomizationHelpers::MakeBrowseButton( FSimpleDelegate::CreateSP( this, &SDataprepConsumerWidget::OnBrowseContentFolder ) );
+
+	TSharedRef<SWidget> Widget = SNew(SBorder)
+	.BorderImage( FCoreStyle::Get().GetBrush("NoBrush") )
+	[
+		SNew(SGridPanel)
+		.FillColumn(0, 1.0f)
+		+ SGridPanel::Slot(0, 0)
+		.Padding(10.0f, 5.0f, 0.0f, 5.0f)
+		[
+			SNew(SHorizontalBox)
+			+ SHorizontalBox::Slot()
+			.FillWidth(1.f)
+			.VAlign(VAlign_Center)
+			.HAlign(HAlign_Left)
+			[
+				SNew(SSplitter)
+				.Style(FEditorStyle::Get(), "DetailsView.Splitter")
+				.PhysicalSplitterHandleSize(1.0f)
+				.HitDetectionSplitterHandleSize(5.0f)
+				.ResizeMode(ESplitterResizeMode::Fill)
+				+ SSplitter::Slot()
+				.Value(ColumnSizeData->LeftColumnWidth)
+				.OnSlotResized(SSplitter::FOnSlotResized::CreateSP(this, &SDataprepConsumerWidget::OnLeftColumnResized))
+				[
+					SNew(SHorizontalBox)
+					.Clipping(EWidgetClipping::OnDemand)
+					+ SHorizontalBox::Slot()
+					.HAlign(HAlign_Left)
+					.VAlign(VAlign_Center)
+					[
+						SNew(STextBlock)
+						.Text(LOCTEXT("DataprepSlateHelper_ContentFolderLabel", "Folder"))
+						.Font(IDetailLayoutBuilder::GetDetailFont())
+					]
+				]
+				+ SSplitter::Slot()
+				.Value(ColumnSizeData->RightColumnWidth)
+				.OnSlotResized(ColumnSizeData->OnWidthChanged)
+				[
+					SNew(SHorizontalBox)
+					.Clipping(EWidgetClipping::OnDemand)
+					+ SHorizontalBox::Slot()
+					.FillWidth(1.0f)
+					.Padding(5.0f, 2.5f, 2.0f, 2.5f)
+					[
+						// Trick to force the splitter widget to fill up the space of its parent
+						// Strongly inspired from SDetailSingleItemRow
+						SNew(DataprepWidgetUtils::SConstrainedBox)
+						[
+							SAssignNew(ContentFolderTextBox, SEditableTextBox)
+							.Font(IDetailLayoutBuilder::GetDetailFont())
+							.HintText(LOCTEXT("DataprepSlateHelper_ContentFolderHintText", "Set the content folder to save in"))
+							.IsReadOnly(false)
+							.OnTextCommitted(FOnTextCommitted::CreateSP(this, &SDataprepConsumerWidget::OnTextCommitted))
+						]
+					]
+				]
+			]
+		]
+		+ SGridPanel::Slot(1, 0)
+		.Padding(5.0f, 5.0f, 0.0f, 5.0f)
+		[
+			SNew(SHorizontalBox)
+			+ SHorizontalBox::Slot()
+			.VAlign(VAlign_Center)
+			.HAlign(HAlign_Right)
+			.AutoWidth()
+			[
+				BrowseButton
+			]
+		]
+		// This column is required to align consumer section with producers' section which has 3 columns.
+		+ SGridPanel::Slot(2, 0)
+		.Padding(5.0f, 5.0f, 0.0f, 5.0f)
+		[
+			SNew(SHorizontalBox)
+			+ SHorizontalBox::Slot()
+			.VAlign(VAlign_Center)
+			.HAlign(HAlign_Right)
+			.AutoWidth()
+			[
+				SNew(SButton)
+				.IsFocusable(false)
+				.Visibility(EVisibility::Hidden)
+				.IsEnabled(false)
+				.VAlign(VAlign_Top)
+				.Content()
+				[
+					SNew(STextBlock)
+					.Font(FDataprepEditorUtils::GetGlyphFont())
+					.ColorAndOpacity(FLinearColor::Transparent)
+					.Text(FEditorFontGlyphs::Exclamation_Triangle)
+				]
+			]
+		]
+		+ SGridPanel::Slot(0, 1)
+		.Padding(10.0f, 5.0f, 0.0f, 5.0f)
+		[
+			SNew(SHorizontalBox)
+			+ SHorizontalBox::Slot()
+			.FillWidth(1.f)
+			.VAlign(VAlign_Center)
+			.HAlign(HAlign_Left)
+			[
+				SNew(SSplitter)
+				.Style(FEditorStyle::Get(), "DetailsView.Splitter")
+				.PhysicalSplitterHandleSize(1.0f)
+				.HitDetectionSplitterHandleSize(5.0f)
+				.ResizeMode(ESplitterResizeMode::Fill)
+				+ SSplitter::Slot()
+				.Value(ColumnSizeData->LeftColumnWidth)
+				.OnSlotResized(SSplitter::FOnSlotResized::CreateSP(this, &SDataprepConsumerWidget::OnLeftColumnResized))
+				[
+					SNew(SHorizontalBox)
+					.Clipping(EWidgetClipping::OnDemand)
+					+ SHorizontalBox::Slot()
+					.HAlign(HAlign_Left)
+					.VAlign(VAlign_Center)
+					[
+						SNew(STextBlock)
+						.Text(LOCTEXT("DataprepSlateHelper_LevelNameLabel", "Sub-Level"))
+						.Font(IDetailLayoutBuilder::GetDetailFont())
+					]
+				]
+				+ SSplitter::Slot()
+				.Value(ColumnSizeData->RightColumnWidth)
+				.OnSlotResized(ColumnSizeData->OnWidthChanged)
+				[
+					SNew(SHorizontalBox)
+					.Clipping(EWidgetClipping::OnDemand)
+					+ SHorizontalBox::Slot()
+					.FillWidth(1.0f)
+					.Padding(5.0f, 2.5f, 2.0f, 2.5f)
+					[
+						// Trick to force the splitter widget to fill up the space of its parent
+						// Strongly inspired from SDetailSingleItemRow
+						SNew(DataprepWidgetUtils::SConstrainedBox)
+						[
+							SAssignNew(LevelTextBox, SEditableTextBox)
+							.Font(IDetailLayoutBuilder::GetDetailFont())
+							.HintText(LOCTEXT("DataprepSlateHelper_LevelNameHintText", "Current will be used"))
+							.OnTextCommitted(this, &SDataprepConsumerWidget::OnLevelNameChanged)
+						]
+					]
+				]
+			]
+		]
+	];
+
+	return Widget;
+}
+
+TSharedRef<SWidget> SDataprepConsumerWidget::BuildNullWidget()
+{
+	return SNew(SBorder)
+	.BorderImage( FCoreStyle::Get().GetBrush("ToolPanel.GroupBorder") )
+	[
+		SNew(SVerticalBox)
+		+ SVerticalBox::Slot()
+		.Padding(5.0f)
+		[
+			SNew(SHorizontalBox)
+			+ SHorizontalBox::Slot()
+			.HAlign(HAlign_Left)
+			.VAlign(VAlign_Center)
+			.FillWidth(1.0f)
+			[
+				SNew(STextBlock)
+				.Font( IDetailLayoutBuilder::GetDetailFont() )
+				.Text(LOCTEXT("DataprepSlateHelper_Error_InvalidConsumer", "Error: Not a valid consumer"))
+				.Margin(FMargin( 5.0f, 5.0f, 0.0f, 0.0f ) )
+				.ColorAndOpacity( FLinearColor(1,0,0,1) )
+			]
+		]
+	];
+}
+
+void SDataprepConsumerWidget::UpdateContentFolderText()
+{
+	if(UDataprepContentConsumer* Consumer = DataprepConsumerPtr.Get())
+	{
+		FString TargetContentFolder( Consumer->GetTargetContentFolder() );
+
+		if( TargetContentFolder.IsEmpty() )
+		{
+			TargetContentFolder = TEXT("/Content");
+		}
+		else if( TargetContentFolder.StartsWith( TEXT( "/Game" ) ) )
+		{
+			TargetContentFolder = TargetContentFolder.Replace( TEXT( "/Game" ), TEXT( "/Content" ) );
+		}
+
+		ContentFolderTextBox->SetText( FText::FromString( TargetContentFolder + TEXT( "/" ) ) );
+	}
+	else
+	{
+		ContentFolderTextBox->SetText( TAttribute<FText>() );
+	}
+}
+
+void SDataprepConsumerWidget::OnConsumerChanged()
+{
+	if(DataprepConsumerPtr.IsValid())
+	{
+		UpdateContentFolderText();
+		LevelTextBox->SetText( FText::FromString( DataprepConsumerPtr->GetLevelName() ) );
+	}
+	else
+	{
+		ContentFolderTextBox->SetText( TAttribute<FText>() );
+		LevelTextBox->SetText( TAttribute<FText>() );
+	}
+}
+
+
+void SDataprepDetailsView::CreateDefaultWidget( int32 Index, TSharedPtr< SWidget >& NameWidget, TSharedPtr< SWidget >& ValueWidget, float LeftPadding, EHorizontalAlignment HAlign, EVerticalAlignment VAlign, const FDataprepParameterizationContext& ParameterizationContext)
+{
+	TSharedRef<SHorizontalBox> NameColumn = SNew(SHorizontalBox)
+		.Clipping(EWidgetClipping::OnDemand);
+
+	NameWidget->SetClipping( EWidgetClipping::OnDemand );
+
+	// Add the name widget
+	NameColumn->AddSlot()
+	.VAlign(VAlign_Center)
+	.HAlign(HAlign_Left)
+	.Padding(FMargin(LeftPadding, 0.f, 0.f, 0.f))
+	[
+		NameWidget.ToSharedRef()
+	];
+
+	if ( ParameterizationContext.State == EParametrizationState::IsParameterized )
+	{
+		NameColumn->AddSlot()
+		.HAlign(HAlign_Right)
+		.VAlign(VAlign_Center)
+		.Padding(FMargin(5.f, 0.f, 5.f, 0.f))
+		.AutoWidth()
+		[
+			SNew(SDataprepParameterizationLinkIcon, DataprepAssetForParameterization.Get(), DetailedObjectAsParameterizable, ParameterizationContext.PropertyChain)
+		];
+	}
+
+	FOnContextMenuOpening OnContextMenuOpening;
+
+	if ( ParameterizationContext.State == EParametrizationState::IsParameterized || ParameterizationContext.State == EParametrizationState::CanBeParameterized )
+	{
+		OnContextMenuOpening.BindLambda( [this, InPropertyChain = ParameterizationContext.PropertyChain] () -> TSharedPtr<SWidget>
+			{
+				if ( UDataprepAsset * DataprepAsset = DataprepAssetForParameterization.Get() )
+				{
+					if ( DetailedObjectAsParameterizable )
+					{
+						FMenuBuilder MenuBuilder(true, nullptr);
+						FDataprepEditorUtils::PopulateMenuForParameterization(MenuBuilder, *DataprepAsset, *DetailedObjectAsParameterizable, InPropertyChain);
+						return MenuBuilder.MakeWidget();
+					}
+				}
+				return TSharedPtr<SWidget>();
+			});
+	}
+
+	GridPanel->AddSlot(0, Index)
+	[
+		SNew(SDataprepContextMenuOverride)
+		.OnContextMenuOpening(OnContextMenuOpening)
+		[
+			DataprepWidgetUtils::CreatePropertyWidget( NameColumn, ValueWidget, ColumnSizeData, Spacing, HAlign, VAlign )
+		]
+	];
+
+	if(bColumnPadding)
+	{
+		// Add two more columns to align parameter widget
+		GridPanel->AddSlot(1, Index)
+		.Padding(5.0f, 5.0f, 0.0f, 5.0f)
+		[
+			SNew(SHorizontalBox)
+			+ SHorizontalBox::Slot()
+			.VAlign(VAlign_Center)
+			.HAlign(HAlign_Right)
+			.AutoWidth()
+			[
+				SNew(SButton)
+				.IsFocusable(false)
+				.Visibility(EVisibility::Hidden)
+				.IsEnabled(false)
+				.VAlign(VAlign_Top)
+				.Content()
+				[
+					SNew(STextBlock)
+					.Font(FDataprepEditorUtils::GetGlyphFont())
+					.ColorAndOpacity(FLinearColor::Transparent)
+					.Text(FEditorFontGlyphs::Exclamation_Triangle)
+				]
+			]
+		];
+
+		GridPanel->AddSlot(2, Index)
+		.Padding(5.0f, 5.0f, 0.0f, 5.0f)
+		[
+			SNew(SHorizontalBox)
+			+ SHorizontalBox::Slot()
+			.VAlign(VAlign_Center)
+			.HAlign(HAlign_Right)
+			.AutoWidth()
+			[
+				SNew(SButton)
+				.IsFocusable(false)
+				.Visibility(EVisibility::Hidden)
+				.IsEnabled(false)
+				.VAlign(VAlign_Top)
+				.Content()
+				[
+					SNew(STextBlock)
+					.Font(FDataprepEditorUtils::GetGlyphFont())
+					.ColorAndOpacity(FLinearColor::Transparent)
+					.Text(FEditorFontGlyphs::Exclamation_Triangle)
+				]
+			]
+		];
+	}
+}
+
+void SDataprepDetailsView::OnPropertyChanged(const FPropertyChangedEvent& InEvent)
+{
+	FDataprepEditorUtils::NotifySystemOfChangeInPipeline( DetailedObject );
+
+	if( TrackedProperties.Contains( InEvent.Property ) )
+	{
+		ForceRefresh();
+	}
+}
+
+void SDataprepDetailsView::OnObjectReplaced(const TMap<UObject*, UObject*>& ReplacementObjectMap)
+{
+	if ( UObject * const* ObjectPtr = ReplacementObjectMap.Find( DetailedObject ) )
+	{
+		DetailedObject = *ObjectPtr;
+		if ( DetailedObject->IsA<UDataprepParameterizableObject>() )
+		{
+			DetailedObjectAsParameterizable = static_cast<UDataprepParameterizableObject*>( DetailedObject );
+		}
+
+		ForceRefresh();
+	}
+}
+
+void SDataprepDetailsView::ForceRefresh()
+{
+	// ueent_hotfix Hotfix for 4.24 (Remove the ui flickering)
+	InvalidatePrepass();
+	bRefreshObjectToDisplay = true;
+}
+
+void SDataprepDetailsView::OnDataprepParameterizationStatusForObjectsChanged(const TSet<UObject*>* Objects)
+{
+	if ( !Objects || Objects->Contains( DetailedObjectAsParameterizable ) )
+	{
+		ForceRefresh();
+	}
+}
+
+void SDataprepDetailsView::OnObjectTransacted(UObject* Object, const class FTransactionObjectEvent& TransactionObjectEvent)
+{
+	// Hack to support refresh the parameterization display of a dataprep instance
+	if ( Object == DetailedObject || ( DetailedObject && DetailedObject->GetOuter() == Object ) )
+	{
+		ForceRefresh();
+	}
+}
+
+void SDataprepDetailsView::AddWidgets( const TArray< TSharedRef< IDetailTreeNode > >& DetailTree, int32& Index, float LeftPadding, const FDataprepParameterizationContext& InParameterizationContext)
+{
+	auto IsDetailNodeDisplayable = []( const TSharedPtr< IPropertyHandle >& PropertyHandle)
+	{
+		if(PropertyHandle.IsValid() && PropertyHandle->IsValidHandle() && PropertyHandle->IsEditable())
+		{
+			UProperty* Property = PropertyHandle->GetProperty();
+
+			if ( Property )
+			{
+				if ( UObject * Outer = Property->GetOuter() )
+				{
+					// if the outer is a container property (array,set or map) it's editable even without the proper flags.
+					UClass* OuterClass = Outer->GetClass();
+					if ( OuterClass == UArrayProperty::StaticClass() || OuterClass == USetProperty::StaticClass() || OuterClass == UMapProperty::StaticClass() )
+					{
+						return true;
+					}
+				}
+
+				return Property && !Property->HasAnyPropertyFlags(CPF_DisableEditOnInstance) && Property->HasAnyPropertyFlags(CPF_Edit);
+			}
+		}
+
+		// Ok to display DetailNode without property because at this stage the parent property was displayable
+		return true;
+	};
+
+	auto IsDetailNodeDisplayableContainerProperty = []( const TSharedPtr< IPropertyHandle >& PropertyHandle)
+	{
+		if( PropertyHandle.IsValid() && PropertyHandle->IsValidHandle() && PropertyHandle->IsEditable() )
+		{
+			if ( UProperty* Property = PropertyHandle->GetProperty() )
+			{ 
+				UClass* PropertyClass = Property->GetClass();
+				if ( PropertyClass == UArrayProperty::StaticClass() || PropertyClass == USetProperty::StaticClass() || PropertyClass == UMapProperty::StaticClass() )
+				{
+					return !Property->HasAnyPropertyFlags(CPF_DisableEditOnInstance) && Property->HasAnyPropertyFlags(CPF_Edit);
+				}
+			}
+		}
+
+		return false;
+	};
+
+	for( const TSharedRef< IDetailTreeNode >& ChildNode : DetailTree )
+	{
+		TSharedPtr< IPropertyHandle > PropertyHandle = ChildNode->CreatePropertyHandle();
+		FDataprepParameterizationContext CurrentParameterizationContext = FDataprepParameterizationUtils::CreateContext( PropertyHandle, InParameterizationContext );
+		if ( CurrentParameterizationContext.State == EParametrizationState::CanBeParameterized )
+		{
+			if ( UDataprepAsset* DataprepAsset = DataprepAssetForParameterization.Get() )
+			{
+				if ( DataprepAsset->IsObjectPropertyBinded( DetailedObjectAsParameterizable, CurrentParameterizationContext.PropertyChain ) )
+				{
+					CurrentParameterizationContext.State = EParametrizationState::IsParameterized;
+				}
+			}
+		}
+
+
+		if( ChildNode->GetNodeType() == EDetailNodeType::Category )
+		{
+			if( Index > 0 )
+			{
+				GridPanel->AddSlot( 0, Index )
+				[
+					SNew(SHorizontalBox)
+					+ SHorizontalBox::Slot()
+					.VAlign(VAlign_Center)
+					.HAlign(HAlign_Left)
+					[
+						SNew(SSpacer)
+						.Size( FVector2D( 0.f, 10.f ) )
+					]
+				];
+				Index++;
+			}
+
+			TArray< TSharedRef< IDetailTreeNode > > Children;
+			ChildNode->GetChildren( Children );
+			AddWidgets( Children, Index, LeftPadding, CurrentParameterizationContext );
+		}
+		else if( IsDetailNodeDisplayableContainerProperty( PropertyHandle ) )
+		{
+			TSharedPtr< IDetailPropertyRow > DetailPropertyRow = ChildNode->GetRow();
+			if( DetailPropertyRow.IsValid() )
+			{
+				FDetailWidgetRow Row;
+				TSharedPtr< SWidget > NameWidget;
+				TSharedPtr< SWidget > ValueWidget;
+				DetailPropertyRow->GetDefaultWidgets( NameWidget, ValueWidget, Row, true );
+
+				CreateDefaultWidget( Index, NameWidget, ValueWidget, LeftPadding, Row.ValueWidget.HorizontalAlignment, Row.ValueWidget.VerticalAlignment, CurrentParameterizationContext );
+				Index++;
+
+				TArray< TSharedRef< IDetailTreeNode > > Children;
+				ChildNode->GetChildren( Children );
+				if( Children.Num() > 0 )
+				{
+					// #ueent_todo: Find a way to add collapse/expand capability for property with children
+					AddWidgets( Children, Index, LeftPadding + 10.f, CurrentParameterizationContext );
+				}
+
+				TrackedProperties.Add( PropertyHandle->GetProperty() );
+			}
+		}
+		else if( IsDetailNodeDisplayable( PropertyHandle ) )
+		{
+			TSharedPtr< SWidget > NameWidget;
+			TSharedPtr< SWidget > ValueWidget;
+			EHorizontalAlignment HAlign;
+			EVerticalAlignment VAlign;
+
+			TSharedPtr< IDetailPropertyRow > DetailPropertyRow = ChildNode->GetRow();
+			if( DetailPropertyRow.IsValid() )
+			{
+				FDetailWidgetRow Row;
+				DetailPropertyRow->GetDefaultWidgets( NameWidget, ValueWidget, Row, true );
+				HAlign = Row.ValueWidget.HorizontalAlignment;
+				VAlign = Row.ValueWidget.VerticalAlignment;
+			}
+			else
+			{
+				FNodeWidgets NodeWidgets = ChildNode->CreateNodeWidgets();
+
+				NameWidget = NodeWidgets.NameWidget;
+				ValueWidget = NodeWidgets.ValueWidget;
+				HAlign = NodeWidgets.ValueWidgetLayoutData.HorizontalAlignment;
+				VAlign = NodeWidgets.ValueWidgetLayoutData.VerticalAlignment;
+			}
+
+			if( NameWidget.IsValid() && ValueWidget.IsValid() )
+			{
+				CreateDefaultWidget( Index, NameWidget, ValueWidget, LeftPadding, HAlign, VAlign, CurrentParameterizationContext );
+				Index++;
+
+				bool bDisplayChildren = true;
+
+				// Do not display children if the property is a FVector or FVector2D
+				if( PropertyHandle.IsValid() )
+				{
+					FVector DummyVec;
+					FVector2D DummyVec2D;
+
+					bDisplayChildren &= PropertyHandle->GetValue( DummyVec ) == FPropertyAccess::Fail;
+					bDisplayChildren &= PropertyHandle->GetValue( DummyVec2D ) == FPropertyAccess::Fail;
+				}
+
+				TArray< TSharedRef< IDetailTreeNode > > Children;
+				ChildNode->GetChildren( Children );
+				if( bDisplayChildren && Children.Num() > 0 )
+				{
+					// #ueent_todo: Find a way to add collapse/expand capability for property with children
+					AddWidgets( Children, Index, LeftPadding + 10.f, CurrentParameterizationContext );
+				}
+			}
+		}
+	}
+}
+
+void SDataprepDetailsView::Construct(const FArguments& InArgs)
+{
+	bRefreshObjectToDisplay = false;
+	DetailedObject = InArgs._Object;
+	Spacing = InArgs._Spacing;
+	bColumnPadding = InArgs._ColumnPadding;
+
+	if (InArgs._ColumnSizeData.IsValid())
+	{
+		ColumnSizeData = InArgs._ColumnSizeData;
+	}
+	else
+	{
+		ColumnWidth = 0.7f;
+		ColumnSizeData = MakeShared<FDataprepDetailsViewColumnSizeData>();
+		ColumnSizeData->LeftColumnWidth = TAttribute<float>(this, &SDataprepDetailsView::OnGetLeftColumnWidth);
+		ColumnSizeData->RightColumnWidth = TAttribute<float>(this, &SDataprepDetailsView::OnGetRightColumnWidth);
+		ColumnSizeData->OnWidthChanged = SSplitter::FOnSlotResized::CreateSP(this, &SDataprepDetailsView::OnSetColumnWidth);
+	}
+	
+	FPropertyEditorModule& PropertyEditorModule = FModuleManager::GetModuleChecked<FPropertyEditorModule>("PropertyEditor");
+
+	FPropertyRowGeneratorArgs Args;
+	Generator = PropertyEditorModule.CreatePropertyRowGenerator(Args);
+
+	if( DetailedObject != nullptr )
+	{
+		TArray< UObject* > Objects;
+		Objects.Add( DetailedObject );
+		Generator->SetObjects( Objects );
+	}
+
+	OnPropertyChangedHandle = Generator->OnFinishedChangingProperties().AddSP( this, &SDataprepDetailsView::OnPropertyChanged );
+
+	if ( GEditor )
+	{
+		OnObjectReplacedHandle = GEditor->OnObjectsReplaced().AddSP(this, &SDataprepDetailsView::OnObjectReplaced);
+		OnObjectTransactedHandle = FCoreUObjectDelegates::OnObjectTransacted.AddSP( this, &SDataprepDetailsView::OnObjectTransacted );
+	}
+
+	Construct();
+}
+
+void SDataprepDetailsView::Construct()
+{
+	DataprepAssetForParameterization.Reset();
+	bHasCustomPrepass = true;
+
+	if ( DetailedObject )
+	{
+		UDataprepAsset* DataprepAsset = FDataprepParameterizationUtils::GetDataprepAssetForParameterization( DetailedObject );
+		if ( DataprepAsset )
+		{
+			OnDataprepParameterizationStatusForObjectsChangedHandle = DataprepAsset->OnParameterizedObjectsChanged.AddSP( this, &SDataprepDetailsView::OnDataprepParameterizationStatusForObjectsChanged );
+		}
+
+		if ( DetailedObject->IsA<UDataprepParameterizableObject>() )
+		{
+			DetailedObjectAsParameterizable = static_cast<UDataprepParameterizableObject*>(DetailedObject);
+		}
+
+		FDataprepParameterizationContext ParameterizationContext;
+		ParameterizationContext.State = DataprepAsset && DetailedObjectAsParameterizable ? EParametrizationState::CanBeParameterized : EParametrizationState::InvalidForParameterization;
+		DataprepAssetForParameterization = DataprepAsset;
+
+		GridPanel = SNew(SGridPanel).FillColumn( 0.0f, 1.0f );
+		TArray< TSharedRef< IDetailTreeNode > > RootNodes = Generator->GetRootTreeNodes();
+
+		int32 Index = 0;
+		AddWidgets(RootNodes, Index, 0.f, ParameterizationContext);
+
+		ChildSlot
+		[
+			SNew(SVerticalBox)
+			+ SVerticalBox::Slot()
+			.Padding(5.0f)
+			[
+				GridPanel.ToSharedRef()
+			]
+		];
+	}
+	else
+	{
+		FText ErrorText = LOCTEXT( "DataprepSlateHelper_InvalidDetailedObject", "Error: Not a valid Object" );
+
+		ChildSlot
+		[
+			SNew(SVerticalBox)
+			+ SVerticalBox::Slot()
+			.Padding(5.0f)
+			[
+				SNew(SHorizontalBox)
+				+ SHorizontalBox::Slot()
+				.HAlign(HAlign_Left)
+				.VAlign(VAlign_Center)
+				.FillWidth(1.0f)
+				[
+					SNew(STextBlock)
+					.Font( IDetailLayoutBuilder::GetDetailFontBold() )
+					.Text( ErrorText )
+					.Margin(FMargin( 5.0f, 5.0f, 0.0f, 0.0f ) )
+					.ColorAndOpacity( FLinearColor(1,0,0,1) )
+				]
+			]
+		];
+	}
+}
+
+SDataprepDetailsView::~SDataprepDetailsView()
+{
+	Generator->OnFinishedChangingProperties().Remove( OnPropertyChangedHandle );
+
+	if ( GEditor )
+	{
+		GEditor->OnObjectsReplaced().Remove( OnObjectReplacedHandle );
+		FCoreUObjectDelegates::OnObjectTransacted.Remove( OnObjectTransactedHandle );
+	}
+
+	if ( UDataprepAsset* DataprepAsset = DataprepAssetForParameterization.Get() )
+	{
+		DataprepAsset->OnParameterizedObjectsChanged.Remove( OnDataprepParameterizationStatusForObjectsChangedHandle );
+	}
+}
+
+void SDataprepDetailsView::SetObjectToDisplay(UObject& Object)
+{
+	UObject* NewObjectToDisplay = &Object;
+	if ( DetailedObject != NewObjectToDisplay )
+	{
+		DetailedObject = NewObjectToDisplay;
+		if ( DetailedObject->IsA<UDataprepParameterizableObject>() )
+		{
+			DetailedObjectAsParameterizable = static_cast<UDataprepParameterizableObject*>( DetailedObject );
+		}
+		ForceRefresh();
+	}
+}
+
+void SDataprepDetailsView::AddReferencedObjects(FReferenceCollector& Collector)
+{
+	Collector.AddReferencedObject( DetailedObject );
+	Collector.AddReferencedObject( DetailedObjectAsParameterizable );
+	Collector.AddReferencedObjects( TrackedProperties );
+}
+
+bool SDataprepDetailsView::CustomPrepass(float LayoutScaleMultiplier)
+{
+	if ( bRefreshObjectToDisplay )
+	{
+		TArray< UObject* > Objects;
+		Objects.Add( DetailedObject );
+		Generator->SetObjects( Objects );
+		Construct();
+		bRefreshObjectToDisplay = false;
+	}
+
+	return true;
+}
+
+void SDataprepContextMenuOverride::Construct(const FArguments& InArgs)
+{
+	OnContextMenuOpening = InArgs._OnContextMenuOpening;
+	
+	ChildSlot
+	[
+		InArgs._DefaultSlot.Widget
+	];
+}
+
+FReply SDataprepContextMenuOverride::OnMouseButtonUp(const FGeometry& MyGeometry, const FPointerEvent& MouseEvent)
+{
+	if ( MouseEvent.GetEffectingButton() == EKeys::RightMouseButton && OnContextMenuOpening.IsBound() )
+	{
+		TSharedPtr<SWidget> ContextMenu = OnContextMenuOpening.Execute();
+		if ( ContextMenu )
+		{ 
+			FWidgetPath WidgetPath = MouseEvent.GetEventPath() ? *MouseEvent.GetEventPath() : FWidgetPath();
+			FSlateApplication::Get().PushMenu(
+				AsShared(),
+				WidgetPath,
+				ContextMenu.ToSharedRef(),
+				MouseEvent.GetScreenSpacePosition(),
+				FPopupTransitionEffect(FPopupTransitionEffect::ContextMenu)
+				);
+
+			return FReply::Handled();
+		}
+	}
+
+	return FReply::Unhandled();
+}
+
+void SDataprepInstanceParentWidget::Construct(const FArguments& InArgs)
+{
+	DataprepInstancePtr = InArgs._DataprepInstance;
+	if(!DataprepInstancePtr.IsValid())
+	{
+		return;
+	}
+
+	if (InArgs._ColumnSizeData.IsValid())
+	{
+		ColumnSizeData = InArgs._ColumnSizeData;
+	}
+	else
+	{
+		ColumnWidth = 0.7f;
+		ColumnSizeData = MakeShared<FDataprepDetailsViewColumnSizeData>();
+		ColumnSizeData->LeftColumnWidth = TAttribute<float>(this, &SDataprepInstanceParentWidget::OnGetLeftColumnWidth);
+		ColumnSizeData->RightColumnWidth = TAttribute<float>(this, &SDataprepInstanceParentWidget::OnGetRightColumnWidth);
+		ColumnSizeData->OnWidthChanged = SSplitter::FOnSlotResized::CreateSP(this, &SDataprepInstanceParentWidget::OnSetColumnWidth);
+	}
+
+
+	TSharedRef<SWidget> NameWidget = SNew(SHorizontalBox)
+	.Clipping(EWidgetClipping::OnDemand)
+	+ SHorizontalBox::Slot()
+	.VAlign(VAlign_Center)
+	.HAlign(HAlign_Left)
+	.Padding(FMargin(0.f, 0.f, 0.f, 0.f))
+	[
+		SNew( STextBlock )
+		.Text( LOCTEXT("DataprepInstanceParentWidget_Parent_Label", "Parent") )
+		.Font( IDetailLayoutBuilder::GetDetailFont() )
+	];
+
+	TSharedRef<SWidget> ValueWidget = SNew( SObjectPropertyEntryBox )
+	.AllowedClass( UDataprepAssetInterface::StaticClass() )
+	.OnObjectChanged( this, &SDataprepInstanceParentWidget::SetDataprepInstanceParent )
+	.OnShouldFilterAsset( this, &SDataprepInstanceParentWidget::ShouldFilterAsset )
+	.ObjectPath( this, &SDataprepInstanceParentWidget::GetDataprepInstanceParent );
+
+	// The widget is disabled as the workflow to change the parent asset of a Dataprep instance is refined
+	ValueWidget->SetEnabled( false );
+
+	ChildSlot
+	[
+		SNew(SHorizontalBox)
+		+SHorizontalBox::Slot()
+		.FillWidth(1)
+		[
+			SNew(SVerticalBox)
+			+ SVerticalBox::Slot()
+			.AutoHeight()
+			[
+				SNew(SHorizontalBox)
+				+ SHorizontalBox::Slot()
+				[
+					SNew(SSpacer)
+					.Size( FVector2D( 200, 10 ) )		
+				]
+			]
+			+ SVerticalBox::Slot()
+			.AutoHeight()
+			[
+				DataprepWidgetUtils::CreatePropertyWidget( NameWidget, ValueWidget, ColumnSizeData, 0.0f )
+			]
+		]
+	];
+}
+
+void SDataprepInstanceParentWidget::SetDataprepInstanceParent(const FAssetData& InAssetData)
+{
+	if(UDataprepAssetInstance* DataprepInstance = DataprepInstancePtr.Get())
+	{
+		if ( UDataprepAsset* DataprepAsset = Cast<UDataprepAsset>( InAssetData.GetAsset() ) )
+		{
+			const FScopedTransaction Transaction( LOCTEXT("DataprepInstance_SetParent", "Set Parent") );
+			DataprepInstance->SetParent( DataprepAsset );
+		}
+	}
+}
+
+FString SDataprepInstanceParentWidget::GetDataprepInstanceParent() const
+{
+	FString PathName;
+
+	if(UDataprepAssetInstance* DataprepInstance = DataprepInstancePtr.Get())
+	{
+		if(DataprepInstance->GetParent())
+		{
+			PathName = DataprepInstance->GetParent()->GetPathName();
+		}
+	}
+
+	return PathName;
+}
+
+bool SDataprepInstanceParentWidget::ShouldFilterAsset(const FAssetData& InAssetData)
+{
+	if(UDataprepAssetInstance* DataprepInstance = DataprepInstancePtr.Get())
+	{
+		if ( InAssetData.GetClass() == UDataprepAssetInterface::StaticClass() )
+		{
+			FAssetData CurrentAssetData(DataprepInstance->GetParent());
+			return CurrentAssetData != InAssetData;
+		}
+	}
+
+	return false;
+}
+
+TSharedRef<SWidget> DataprepWidgetUtils::CreateParameterRow( TSharedPtr<SWidget> ParameterWidget )
+{
+	return  SNew(SGridPanel)
+	.FillColumn(0, 1.0f)
+	+ SGridPanel::Slot(0, 0)
+	.Padding(10.0f, 5.0f, 0.0f, 5.0f)
+	[
+		ParameterWidget.ToSharedRef()
+	]
+	// Add two more columns to align parameter widget
+	+ SGridPanel::Slot(1, 0)
+	.Padding(5.0f, 5.0f, 0.0f, 5.0f)
+	[
+		SNew(SHorizontalBox)
+		+ SHorizontalBox::Slot()
+		.VAlign(VAlign_Center)
+		.HAlign(HAlign_Right)
+		.AutoWidth()
+		[
+			SNew(SButton)
+			.IsFocusable(false)
+			.Visibility(EVisibility::Hidden)
+			.IsEnabled(false)
+			.VAlign(VAlign_Top)
+			.Content()
+			[
+				SNew(STextBlock)
+				.Font(FDataprepEditorUtils::GetGlyphFont())
+				.ColorAndOpacity(FLinearColor::Transparent)
+				.Text(FEditorFontGlyphs::Exclamation_Triangle)
+			]
+		]
+	]
+	+ SGridPanel::Slot(2, 0)
+	.Padding(5.0f, 5.0f, 0.0f, 5.0f)
+	[
+		SNew(SHorizontalBox)
+		+ SHorizontalBox::Slot()
+		.VAlign(VAlign_Center)
+		.HAlign(HAlign_Right)
+		.AutoWidth()
+		[
+			SNew(SButton)
+			.IsFocusable(false)
+			.Visibility(EVisibility::Hidden)
+			.IsEnabled(false)
+			.VAlign(VAlign_Top)
+			.Content()
+			[
+				SNew(STextBlock)
+				.Font(FDataprepEditorUtils::GetGlyphFont())
+				.ColorAndOpacity(FLinearColor::Transparent)
+				.Text(FEditorFontGlyphs::Exclamation_Triangle)
+			]
+		]
+	];
+}
+
+#undef LOCTEXT_NAMESPACE
diff --git a/Engine/Plugins/Enterprise/DataprepEditor/Source/DataprepEditor/Private/Widgets/DataprepWidgets.h b/Engine/Plugins/Enterprise/DataprepEditor/Source/DataprepEditor/Private/Widgets/DataprepWidgets.h
new file mode 100644
--- /dev/null
+++ b/Engine/Plugins/Enterprise/DataprepEditor/Source/DataprepEditor/Private/Widgets/DataprepWidgets.h
@@ -0,0 +1,346 @@
+// Copyright 1998-2019 Epic Games, Inc. All Rights Reserved.
+
+#pragma once
+
+#include "CoreMinimal.h"
+
+#include "Parameterization/DataprepParameterizationUtils.h"
+
+#include "Framework/Application/SlateApplication.h"
+#include "Framework/SlateDelegates.h"
+#include "Widgets/Layout/SSplitter.h"
+#include "Widgets/SCompoundWidget.h"
+#include "Widgets/Views/SListView.h"
+#include "Widgets/Views/STreeView.h"
+
+class SEditableTextBox;
+class SGridPanel;
+class UDataprepAsset;
+class UDataprepAssetInstance;
+class UDataprepContentConsumer;
+class UDataprepParameterizableObject;
+
+struct FAssetData;
+
+namespace DataprepWidgetUtils
+{
+	TSharedRef<SWidget> CreateParameterRow( TSharedPtr<SWidget> ParameterWidget );
+
+	/** Helper class to force a widget to fill in a space. Copied from SDetailSingleItemRow.cpp */
+	class SConstrainedBox : public SCompoundWidget
+	{
+	public:
+		SLATE_BEGIN_ARGS(SConstrainedBox) {}
+		SLATE_DEFAULT_SLOT(FArguments, Content)
+			SLATE_END_ARGS()
+
+			void Construct(const FArguments& InArgs)
+		{
+			ChildSlot
+				[
+					InArgs._Content.Widget
+				];
+		}
+
+		virtual FVector2D ComputeDesiredSize(float LayoutScaleMultiplier) const override
+		{
+			// Voluntarily ridiculously large value to force the child widget to fill up the available space
+			const float MinWidthVal = 2000;
+			const FVector2D ChildSize = ChildSlot.GetWidget()->GetDesiredSize();
+			return FVector2D(FMath::Max(MinWidthVal, ChildSize.X), ChildSize.Y);
+		}
+	};
+}
+
+struct FDataprepDetailsViewColumnSizeData
+{
+	TAttribute<float> LeftColumnWidth;
+	TAttribute<float> RightColumnWidth;
+	SSplitter::FOnSlotResized OnWidthChanged;
+
+	void SetColumnWidth(float InWidth) { OnWidthChanged.ExecuteIfBound(InWidth); }
+};
+
+enum class EDataprepCategory 
+{ 
+	Producers, 
+	Consumers, 
+	Parameterization 
+};
+
+typedef STreeView< TSharedRef<EDataprepCategory> > SDataprepCategoryTree;
+
+class SDataprepCategoryWidget : public STableRow< TSharedPtr< EDataprepCategory > >
+{
+public:
+
+	SLATE_BEGIN_ARGS(SDataprepCategoryWidget) {}
+	SLATE_ARGUMENT(TSharedPtr< FDataprepDetailsViewColumnSizeData >, ColumnSizeData)
+	SLATE_ARGUMENT(FText, Title)
+	SLATE_ARGUMENT(TSharedPtr< SWidget >, TitleDetail)
+	SLATE_END_ARGS()
+
+	void Construct( const FArguments& InArgs, TSharedRef< SWidget > InContent, const TSharedRef<STableViewBase>& InOwnerTableView );
+
+	virtual int32 DoesItemHaveChildren() const override { return 1; }
+	virtual bool IsItemExpanded() const override { return bIsExpanded; }
+	virtual void ToggleExpansion() override;
+
+	const FSlateBrush* GetBackgroundImage() const;
+private:
+	bool bIsExpanded = true;
+	TSharedPtr< SWidget > CategoryContent;
+};
+
+class SDataprepConsumerWidget : public SCompoundWidget
+{
+public:
+	SLATE_BEGIN_ARGS(SDataprepConsumerWidget)
+		: _DataprepConsumer(nullptr)
+	{}
+
+	SLATE_ARGUMENT(TSharedPtr< FDataprepDetailsViewColumnSizeData >, ColumnSizeData)
+	SLATE_ARGUMENT(UDataprepContentConsumer*, DataprepConsumer)
+	SLATE_END_ARGS()
+
+	void Construct( const FArguments& InArgs );
+
+	virtual ~SDataprepConsumerWidget();
+
+	/** Update DataprepConsumerPtr and content folder's and level name's text boxes */
+	void SetDataprepConsumer( UDataprepContentConsumer* DataprepConsumer );
+
+private:
+	/** Create complete widget from valid UDataprepContentConsumer object */
+	TSharedRef<SWidget> BuildWidget();
+	/** Create empty widget from invalid UDataprepContentConsumer object */
+	TSharedRef<SWidget> BuildNullWidget();
+	/** Callback when content of level name's text box has changed */
+	void OnLevelNameChanged( const FText& NewLevelName, ETextCommit::Type CommitType );
+	/** Callback when browser's button is clicked */
+	void OnBrowseContentFolder();
+	/** Callback when content of content folder's text box has changed */
+	void OnTextCommitted( const FText&, ETextCommit::Type );
+
+	/** Update content folder's and level name's text boxes */
+	void OnConsumerChanged();
+	/** Update content folder's text boxes */
+	void UpdateContentFolderText();
+
+	/** Callbacks to update splitter */
+	void OnLeftColumnResized(float InNewWidth)
+	{
+		// This has to be bound or the splitter will take it upon itself to determine the size
+		// We do nothing here because it is handled by the column size data
+	}
+
+	float OnGetLeftColumnWidth() const { return 1.0f - ColumnWidth; }
+	float OnGetRightColumnWidth() const { return ColumnWidth; }
+	void OnSetColumnWidth(float InWidth) { ColumnWidth = InWidth; }
+
+private:
+	/** Weak pointer on edited consumer */
+	TWeakObjectPtr<UDataprepContentConsumer> DataprepConsumerPtr;
+	/** Content folder's text box */
+	TSharedPtr< SEditableTextBox > ContentFolderTextBox;
+	/** Level name's text box */
+	TSharedPtr< SEditableTextBox > LevelTextBox;
+	/** Helps sync column resizing with another part of the UI (producers widget) */
+	TSharedPtr< FDataprepDetailsViewColumnSizeData > ColumnSizeData;
+	/** Relative width to control splitters. */
+	float ColumnWidth;
+
+	FDelegateHandle OnConsumerChangedHandle;
+};
+
+class SDataprepDetailsView : public SCompoundWidget, public FGCObject
+{
+	using Super = SCompoundWidget;
+
+public:
+	SLATE_BEGIN_ARGS(SDataprepDetailsView)
+		: _Object( nullptr )
+		, _Spacing(0)
+		, _ColumnPadding(false)
+		{}
+		SLATE_ARGUMENT(TSharedPtr<FDataprepDetailsViewColumnSizeData>, ColumnSizeData)
+		SLATE_ARGUMENT(UObject*, Object)
+		SLATE_ARGUMENT(float, Spacing)
+		SLATE_ARGUMENT(bool, ColumnPadding)
+	SLATE_END_ARGS()
+
+public:
+	void Construct(const FArguments& InArgs);
+
+	~SDataprepDetailsView();
+
+	void SetObjectToDisplay(UObject& Object);
+
+	void ForceRefresh();
+
+	// FGCObject interface
+	virtual void AddReferencedObjects(FReferenceCollector& Collector) override;
+	// End of FGCObject interface
+
+protected:
+	// ueent_hotfix Hack for 4.24 allow to refresh the ui in between two frame without any flickering
+	virtual bool CustomPrepass(float LayoutScaleMultiplier) override;
+
+	// Begin SWidget overrides.
+	// #ueent_todo: This is temporary until we find a better solution to the splitter issue
+	//				See SConstrainedBox's trick in cpp file
+	virtual FVector2D ComputeDesiredSize(float) const override
+	{
+		const float MaxWidth = 400.0f;
+		const FVector2D ChildSize = ChildSlot.GetWidget()->GetDesiredSize();
+		return FVector2D( MaxWidth, ChildSize.Y );
+	}
+	// End SWidget overrides.
+
+private:
+	/** Fills up the details view with the detail nodes created by the property row manager */
+	void Construct();
+
+	/** Add widgets held by array of DetailTreeNode objects */
+	void AddWidgets( const TArray< TSharedRef< class IDetailTreeNode > >& DetailTree, int32& Index, float LeftPadding, const FDataprepParameterizationContext& ParameterizationContext);
+
+	/**
+	 * Inserts a generic widget for a property row into the grid panel
+	 * @param Index						Row index in the grid panel
+	 * @param NameWidget				The widget used in the name column
+	 * @param ValueWidget				The widget used in the value column
+	 * @param LeftPadding				The Padding on the left of the row
+	 * @param ParameterizationContext	Parameterization context of the associated property
+	 */
+	void CreateDefaultWidget(int32 Index, TSharedPtr< SWidget >& NameWidget, TSharedPtr< SWidget >& ValueWidget, float LeftPadding, EHorizontalAlignment HAlign, EVerticalAlignment VAlign, const FDataprepParameterizationContext& ParameterizationContext);
+
+	/** Callback used by all splitters in the details view, so that they move in sync */
+	void OnLeftColumnResized(float InNewWidth)
+	{
+		// This has to be bound or the splitter will take it upon itself to determine the size
+		// We do nothing here because it is handled by the column size data
+	}
+
+	float OnGetLeftColumnWidth() const { return 1.0f - ColumnWidth; }
+	float OnGetRightColumnWidth() const { return ColumnWidth; }
+	void OnSetColumnWidth(float InWidth) { ColumnWidth = InWidth; }
+
+	/** Callback to track property changes on array properties */
+	void OnPropertyChanged( const struct FPropertyChangedEvent& InEvent );
+
+	/** Callback used to detect the existence of a new object to display after a reinstancing process */
+	void OnObjectReplaced(const TMap<UObject*, UObject*>& ReplacementObjectMap);
+
+	void OnDataprepParameterizationStatusForObjectsChanged(const TSet<UObject*>* Objects);
+
+	void OnObjectTransacted(UObject* Object, const class FTransactionObjectEvent& TransactionObjectEvent);
+
+private:
+	/** Row generator applied on detailed object */
+	TSharedPtr< class IPropertyRowGenerator > Generator;
+
+	/** Object to be detailed */
+	UObject* DetailedObject;
+
+	/** Not null if the detailed object is parameterizable */
+	UDataprepParameterizableObject* DetailedObjectAsParameterizable;
+
+	/** Array properties tracked for changes */
+	TSet< UProperty* > TrackedProperties;
+	
+	/** Delegate handle to track property changes on array properties */
+	FDelegateHandle OnPropertyChangedHandle;
+
+	/** Delegate handle to track new object after a re-instancing process */
+	FDelegateHandle OnObjectReplacedHandle;
+
+	/** Delegate handle to track when a object was transacted */
+	FDelegateHandle OnObjectTransactedHandle;
+
+	/** Relative width to control splitters. */
+	float ColumnWidth;
+
+	/** Points to the currently used column size data. Can be provided via argument as well. */
+	TSharedPtr< FDataprepDetailsViewColumnSizeData > ColumnSizeData;
+
+	/** If there is a new object to display on the next tick */
+	uint8 bRefreshObjectToDisplay : 1;
+
+	/**
+	 * Pointer to a Dataprep asset used by the parameterization system.
+	 * It should be nullptr when the parameterization shouldn't be shown
+	 */
+	TWeakObjectPtr<UDataprepAsset> DataprepAssetForParameterization;
+
+	/** Callback used when parameterization has changed */
+	FDelegateHandle OnDataprepParameterizationStatusForObjectsChangedHandle;
+
+	/** Spacing between rows of the SDataprepDetailsView. 0 by default. */
+	float Spacing;
+
+	/**
+	 * Indicates if two columns should be added
+	 * This is useful when the SDataprepDetailsView widget is used along the Producers widget
+	 */
+	bool bColumnPadding;
+
+	/** Grid panel storing the row widgets */
+	TSharedPtr<SGridPanel> GridPanel;
+};
+
+/**
+ * This is widget simply exist to open contextual menu
+ */
+class SDataprepContextMenuOverride : public SCompoundWidget
+{
+public:
+	SLATE_BEGIN_ARGS(SDataprepContextMenuOverride) {}
+		SLATE_DEFAULT_SLOT(FArguments, DefaultSlot)
+		SLATE_EVENT(FOnContextMenuOpening, OnContextMenuOpening)
+	SLATE_END_ARGS()
+
+	void Construct(const FArguments& InArgs);
+
+	virtual FReply OnMouseButtonUp(const FGeometry& MyGeometry, const FPointerEvent& MouseEvent) override;
+
+private:
+	FOnContextMenuOpening OnContextMenuOpening;
+};
+
+// Widget to expose parent asset of a Dataprep instance
+class SDataprepInstanceParentWidget : public SCompoundWidget
+{
+public:
+	SLATE_BEGIN_ARGS(SDataprepInstanceParentWidget) {}
+		SLATE_ARGUMENT(TSharedPtr< FDataprepDetailsViewColumnSizeData >, ColumnSizeData)
+		SLATE_ARGUMENT(UDataprepAssetInstance*, DataprepInstance)
+	SLATE_END_ARGS()
+
+	void Construct(const FArguments& InArgs);
+
+private:
+	void SetDataprepInstanceParent(const FAssetData& InAssetData);
+	FString GetDataprepInstanceParent() const;
+	bool ShouldFilterAsset(const FAssetData& InAssetData);
+
+	/** Callbacks to update splitter */
+	void OnLeftColumnResized(float InNewWidth)
+	{
+		// This has to be bound or the splitter will take it upon itself to determine the size
+		// We do nothing here because it is handled by the column size data
+	}
+
+	float OnGetLeftColumnWidth() const { return 1.0f - ColumnWidth; }
+	float OnGetRightColumnWidth() const { return ColumnWidth; }
+	void OnSetColumnWidth(float InWidth) { ColumnWidth = InWidth; }
+
+private:
+	/** Weak pointer on DataprepAsset instance */
+	TWeakObjectPtr<UDataprepAssetInstance> DataprepInstancePtr;
+
+	/** Helps sync column resizing with another part of the UI */
+	TSharedPtr< FDataprepDetailsViewColumnSizeData > ColumnSizeData;
+
+	/** Relative width to control splitters. */
+	float ColumnWidth;
+};
diff --git a/Engine/Plugins/Enterprise/DataprepEditor/Source/DataprepEditor/Private/Widgets/SAssetsPreviewWidget.cpp b/Engine/Plugins/Enterprise/DataprepEditor/Source/DataprepEditor/Private/Widgets/SAssetsPreviewWidget.cpp
new file mode 100644
--- /dev/null
+++ b/Engine/Plugins/Enterprise/DataprepEditor/Source/DataprepEditor/Private/Widgets/SAssetsPreviewWidget.cpp
@@ -0,0 +1,461 @@
+// Copyright 1998-2019 Epic Games, Inc. All Rights Reserved.
+
+#include "SAssetsPreviewWidget.h"
+
+#include "AssetToolsModule.h"
+#include "IAssetTypeActions.h"
+#include "EditorStyleSet.h"
+#include "Widgets/Images/SImage.h"
+#include "Widgets/Text/STextBlock.h"
+#include "SAssetSearchBox.h"
+#include "Misc/PackageName.h"
+
+#include "UObject/UObjectBaseUtility.h"
+
+#define LOCTEXT_NAMESPACE "AssetPreviewWidget"
+
+namespace AssetPreviewWidget
+{
+	struct FAssetTreeItem
+	{
+		// This is used to accelerate the construction of the tree in the set assets function
+		TMap< FString, FAssetTreeItemPtr> NameToFolder;
+
+		// Childrens
+		TArray< FAssetTreeItemPtr > Folders;
+		TArray< FAssetTreeItemPtr > Assets;
+
+		FString Name;
+		TWeakObjectPtr< UObject > AssetPtr;
+
+		TWeakPtr< SAssetsPreviewWidget > OwnerWeakPtr;
+
+		// This value is cache for the last result for the filter function
+		bool bPassedFilter;
+
+		void AddFolder(FAssetTreeItemPtr Folder)
+		{
+			if (Folder)
+			{
+				NameToFolder.Add(Folder->Name, Folder);
+				Folders.Add(Folder);
+			}
+		}
+
+		bool IsFolder() const
+		{
+			if (Folders.Num() > 0 || Assets.Num() > 0)
+			{
+				return true;
+			}
+			return false;
+		}
+
+		bool Filter(const FText& FilterText)
+		{
+			bPassedFilter = false;
+			if (FilterText.IsEmpty())
+			{
+				bPassedFilter = true;
+			}
+
+			if (IsFolder())
+			{
+				// A folder pass the filter if one of his child pass the filter
+				for (FAssetTreeItemPtr& Folder : Folders)
+				{
+					bPassedFilter = Folder->Filter(FilterText) || bPassedFilter;
+				}
+				for (FAssetTreeItemPtr& Asset : Assets)
+				{
+					bPassedFilter = Asset->Filter(FilterText) || bPassedFilter;
+				}
+			}
+			else if (!bPassedFilter && OwnerWeakPtr.IsValid())
+			{
+				TArray< FString > FilterStrings;
+				FilterText.ToString().ParseIntoArray(FilterStrings, TEXT(" "));
+
+				TSharedPtr< SAssetsPreviewWidget > OwnerWidget = OwnerWeakPtr.Pin();
+				TArray< FString > ItemsName = OwnerWidget->GetItemsName(AssetPtr);
+
+				// All the keywords must be pass for at least one of the items name in the hierarchy
+				bool bPassKeyWord = false;
+				for (const FString& KeyWord : FilterStrings)
+				{
+					bPassKeyWord = false;
+					for (const FString& ItemName : ItemsName)
+					{
+						if (ItemName.Contains(KeyWord))
+						{
+							bPassKeyWord = true;
+							break;
+						}
+					}
+					if (!bPassKeyWord)
+					{
+						break;
+					}
+				}
+
+				bPassedFilter = bPassKeyWord;
+
+			}
+
+			return bPassedFilter;
+		}
+	};
+
+	/** Represents a row in the AssetPreview's tree view */
+	class SAssetPreviewTableRow : public STableRow<FAssetTreeItemPtr>
+	{
+	public:
+		SLATE_BEGIN_ARGS(SAssetPreviewTableRow) {}
+		SLATE_END_ARGS()
+
+		void Construct(const FArguments& InArgs, const TSharedRef<STableViewBase>& OwnerTableView, FAssetTreeItemPtr InItem, TSharedRef< const SAssetsPreviewWidget > InPreviewWidget)
+		{
+			FolderOpenBrush = FEditorStyle::GetBrush("ContentBrowser.AssetTreeFolderOpen");
+			FolderClosedBrush = FEditorStyle::GetBrush("ContentBrowser.AssetTreeFolderClosed");
+			AssetIconBrush = FEditorStyle::GetBrush("ContentBrowser.ColumnViewAssetIcon");
+
+			PreviewWidgetWeakPtr = InPreviewWidget;
+			ItemWeakPtr = InItem;
+
+			FSlateColor IconColor(FLinearColor::White);
+			TWeakObjectPtr< UObject > AssetPtr = InItem->AssetPtr;
+			if (AssetPtr.IsValid())
+			{
+				static FAssetToolsModule& AssetToolsModule = FModuleManager::LoadModuleChecked<FAssetToolsModule>(TEXT("AssetTools"));
+				TSharedPtr<IAssetTypeActions> AssetTypeActions = AssetToolsModule.Get().GetAssetTypeActionsForClass(AssetPtr->GetClass()).Pin();
+
+				if (AssetTypeActions.IsValid())
+				{
+					IconColor = FSlateColor(AssetTypeActions->GetTypeColor());
+				}
+			}
+
+			STableRow::Construct(
+				STableRow::FArguments()
+				.Style(FEditorStyle::Get(), "ContentBrowser.AssetListView.TableRow")
+				.Cursor(EMouseCursor::Default),
+				OwnerTableView);
+
+
+			TSharedRef< SHorizontalBox > Widget = SNew(SHorizontalBox)
+				+ SHorizontalBox::Slot()
+				.AutoWidth()
+				.Padding(0, 0, 2, 0)
+				.VAlign(VAlign_Center)
+				[
+					// Item icon
+					SNew(SImage)
+					.Image(this, &SAssetPreviewTableRow::GetIconBrush)
+					.ColorAndOpacity(IconColor)
+				]
+
+				+ SHorizontalBox::Slot()
+				.AutoWidth()
+				.VAlign(VAlign_Center)
+				[
+					SNew(STextBlock)
+					.Text(FText::FromString(*(InItem->Name)))
+					.Font(FEditorStyle::GetFontStyle("ContentBrowser.SourceTreeItemFont"))
+					.HighlightText(PreviewWidgetWeakPtr.Pin().Get(), &SAssetsPreviewWidget::OnGetHighlightText)
+				];
+
+			SetContent(Widget);
+		}
+
+	private:
+		const FSlateBrush* GetIconBrush() const
+		{
+			const FSlateBrush* IconBrush = AssetIconBrush;
+			if (ItemWeakPtr.Pin()->IsFolder())
+			{
+				const bool bExpanded = PreviewWidgetWeakPtr.Pin()->GetTreeView()->IsItemExpanded(ItemWeakPtr.Pin());
+				IconBrush = bExpanded ? FolderOpenBrush : FolderClosedBrush;
+			}
+			return IconBrush;
+		}
+
+	private:
+		/** Brushes for the different folder states */
+		const FSlateBrush* FolderOpenBrush;
+		const FSlateBrush* FolderClosedBrush;
+		const FSlateBrush* AssetIconBrush;
+
+		FAssetTreeItemWeakPtr ItemWeakPtr;
+
+		/** Weak reference back to the preview widget that owns us */
+		TWeakPtr< const SAssetsPreviewWidget > PreviewWidgetWeakPtr;
+	};
+
+	void SAssetsPreviewWidget::Construct(const FArguments& InArgs)
+	{
+		ChildSlot
+		[
+			SNew(SVerticalBox)
+
+			// Search and commands
+			+ SVerticalBox::Slot()
+			.AutoHeight()
+			.Padding(0, 0, 0, 2)
+			[
+				SNew(SHorizontalBox)
+				// Search
+				+ SHorizontalBox::Slot()
+				.Padding(0, 1, 0, 0)
+				.FillWidth(1.0f)
+				[
+					SNew(SAssetSearchBox)
+					.OnTextChanged(this, &SAssetsPreviewWidget::OnSearchBoxChanged)
+					.OnTextCommitted(this, &SAssetsPreviewWidget::OnSearchBoxCommitted)
+					.DelayChangeNotificationsWhileTyping(true)
+					.HintText(LOCTEXT("SearchHint", "Search..."))
+				]
+			]
+
+			+ SVerticalBox::Slot()
+			//.AutoHeight()
+			.Padding(2.f)
+			[
+				SAssignNew(TreeView, STreeView< FAssetTreeItemPtr >)
+				.SelectionMode( ESelectionMode::Single )
+				.TreeItemsSource(&FilteredRootItems)
+				.OnGenerateRow(this, &SAssetsPreviewWidget::MakeRowWidget)
+				.OnSetExpansionRecursive(this, &SAssetsPreviewWidget::OnSetExpansionRecursive)
+				.OnGetChildren(this, &SAssetsPreviewWidget::OnGetChildren)
+				.OnSelectionChanged(this, &SAssetsPreviewWidget::OnSelectionChangedInternal)
+			]
+		];
+	}
+
+	void SAssetsPreviewWidget::SetAssetsList(const TArray< TWeakObjectPtr< UObject > >& InAssetsList, const FString& InPathToReplace, const FString& InSubstitutePath)
+	{
+		PathToReplace = InPathToReplace;
+		SubstitutePath = InSubstitutePath;
+		RootItems.Empty();
+		TMap< FString, FAssetTreeItemPtr > NamesToRootItem;
+
+		// Make sure the root dir is displayed as "Content".
+		// This is more descriptive for the end user.
+		TArray<FString> Tokens;
+		const FString StrContent = TEXT("Content");
+		SubstitutePath.ParseIntoArray(Tokens, TEXT("/"));
+		if (Tokens.Num() > 0 && !Tokens[0].Equals(StrContent))
+		{
+			const int32 StartChar = SubstitutePath.Find(Tokens[0]);
+			SubstitutePath.RemoveAt(StartChar, Tokens[0].Len(), false);
+			SubstitutePath.InsertAt(StartChar, StrContent);
+		}
+
+		for (const TWeakObjectPtr< UObject >& Asset : InAssetsList)
+		{
+			if (Asset.Get())
+			{
+				TArray< FString > ItemsName = GetItemsName(Asset);
+
+				if (ItemsName.Num() > 0)
+				{
+					FAssetTreeItemPtr LastParent;
+
+					if (ItemsName.Num() > 1)
+					{
+						// Manage the root item
+						{
+							FString ItemName = MoveTemp(ItemsName[0]);
+							FAssetTreeItemPtr RootItem;
+							if (FAssetTreeItemPtr* PtrToTreeItemPtr = NamesToRootItem.Find(ItemName))
+							{
+								RootItem = *PtrToTreeItemPtr;
+							}
+							else
+							{
+								RootItem = MakeShared<FAssetTreeItem>();
+								RootItem->Name = ItemName;
+								RootItem->OwnerWeakPtr = SharedThis(this);
+								RootItems.Add(RootItem);
+								NamesToRootItem.Add(MoveTemp(ItemName), RootItem);
+							}
+							LastParent = MoveTemp(RootItem);
+						}
+
+						//Manage the folders to the asset
+						{
+							const int32 ItemsNameLenghtMinusOne = ItemsName.Num() - 1;
+							for (int32 i = 1; i < ItemsNameLenghtMinusOne; i++)
+							{
+								FString ItemName = MoveTemp(ItemsName[i]);
+								FAssetTreeItemPtr FolderItem;
+								if (FAssetTreeItemPtr* PtrToTreeItemPtr = LastParent->NameToFolder.Find(ItemName))
+								{
+									FolderItem = *PtrToTreeItemPtr;
+								}
+								else
+								{
+									FolderItem = MakeShared<FAssetTreeItem>();
+									FolderItem->Name = MoveTemp(ItemName);
+									FolderItem->OwnerWeakPtr = SharedThis(this);
+									LastParent->AddFolder(FolderItem);
+								}
+
+								LastParent = FolderItem;
+							}
+						}
+					}
+
+					// Create the asset item
+					{
+						FAssetTreeItemPtr AssetItem = MakeShared<FAssetTreeItem>();
+						AssetItem->Name = MoveTemp(ItemsName.Last());
+						AssetItem->AssetPtr = Asset;
+						AssetItem->OwnerWeakPtr = SharedThis(this);
+						LastParent->Assets.Add(AssetItem);
+					}
+				}
+			}
+		}
+
+		// Sort items in alphabetical order
+		SortRecursive(RootItems);
+
+		FilterAssetsNames();
+	}
+
+	void SAssetsPreviewWidget::ClearAssetList()
+	{
+		RootItems.Empty();
+		FilterAssetsNames();
+	}
+
+	void SAssetsPreviewWidget::FilterAssetsNames()
+	{
+		FilteredRootItems.Empty();
+
+		for (FAssetTreeItemPtr& Item : RootItems)
+		{
+			if (Item->Filter(FilterText))
+			{
+				FilteredRootItems.Add(Item);
+			}
+		}
+
+		TreeView->RequestListRefresh();
+
+		ExpandAllFolders();
+	}
+
+	void SAssetsPreviewWidget::ExpandAllFolders()
+	{
+		for (FAssetTreeItemPtr& Item : FilteredRootItems)
+		{
+			ExpandFolderRecursive(Item);
+		}
+	}
+
+	void SAssetsPreviewWidget::ExpandFolderRecursive(FAssetTreeItemPtr InItem)
+	{
+		TreeView->SetItemExpansion(InItem, true);
+		for (FAssetTreeItemPtr& Item : InItem->Folders)
+		{
+			ExpandFolderRecursive(Item);
+		}
+	}
+
+	void SAssetsPreviewWidget::SortRecursive(TArray< FAssetTreeItemPtr >& InItems)
+	{
+		InItems.Sort([](const FAssetTreeItemPtr A, const FAssetTreeItemPtr B) { return A->Name.Compare(B->Name) < 0; });
+
+		for (int ItemIndex = 0; ItemIndex < InItems.Num(); ++ItemIndex)
+		{
+			SortRecursive(InItems[ItemIndex]->Folders);
+			SortRecursive(InItems[ItemIndex]->Assets);
+		}
+	}
+
+	TArray< FString > SAssetsPreviewWidget::GetItemsName(TWeakObjectPtr< UObject > Asset) const
+	{
+		TArray< FString > ItemsName;
+		FString AssetSubPath = Asset->GetPathName(nullptr);
+		if(AssetSubPath.RemoveFromStart(PathToReplace) && !SubstitutePath.IsEmpty())
+		{
+			AssetSubPath = SubstitutePath / AssetSubPath;
+		}
+		AssetSubPath.ReplaceCharInline(TEXT('/'), TEXT('.'));
+		AssetSubPath.ParseIntoArray(ItemsName, TEXT("."), true);
+		return ItemsName;
+	}
+
+	TSharedRef< ITableRow > SAssetsPreviewWidget::MakeRowWidget(FAssetTreeItemPtr InItem, const TSharedRef< STableViewBase >& OwnerTable) const
+	{
+		TSharedPtr< SAssetPreviewTableRow > TableRowWidget;
+
+		SAssignNew(TableRowWidget, SAssetPreviewTableRow, OwnerTable, InItem, SharedThis(this));
+
+		return TableRowWidget.ToSharedRef();
+	}
+
+	void SAssetsPreviewWidget::OnGetChildren(FAssetTreeItemPtr InParent, TArray< FAssetTreeItemPtr >& OutChildren) const
+	{
+		for (FAssetTreeItemPtr& Folder : InParent->Folders)
+		{
+			if (Folder->bPassedFilter)
+			{
+				OutChildren.Add(Folder);
+			}
+		}
+
+		for (FAssetTreeItemPtr& Asset : InParent->Assets)
+		{
+			if (Asset->bPassedFilter)
+			{
+				OutChildren.Add(Asset);
+			}
+		}
+
+	}
+
+	void SAssetsPreviewWidget::OnSearchBoxChanged(const FText& InSearchText)
+	{
+		FilterText = InSearchText;
+		FilterAssetsNames();
+	}
+
+	void SAssetsPreviewWidget::OnSearchBoxCommitted(const FText& InSearchText, ETextCommit::Type CommitInfo)
+	{
+		FilterText = InSearchText;
+		FilterAssetsNames();
+	}
+
+	FText SAssetsPreviewWidget::OnGetHighlightText() const
+	{
+		return FilterText;
+	}
+
+
+	void SAssetsPreviewWidget::OnSetExpansionRecursive(FAssetTreeItemPtr InTreeNode, bool bInIsItemExpanded)
+	{
+		if (InTreeNode.IsValid())
+		{
+			TreeView->SetItemExpansion(InTreeNode, bInIsItemExpanded);
+
+			for (FAssetTreeItemPtr SubFolder : InTreeNode->Folders)
+			{
+				OnSetExpansionRecursive(SubFolder, bInIsItemExpanded);
+			}
+		}
+	}
+
+	void SAssetsPreviewWidget::OnSelectionChangedInternal(FAssetTreeItemPtr ItemSelected, ESelectInfo::Type SelectionType)
+	{
+		if ( ItemSelected )
+		{	
+			TSet< UObject* > Selection;
+			Selection.Add( ItemSelected->AssetPtr.Get() );
+			OnSelectionChanged().Broadcast( Selection );
+		}
+	}
+}
+
+#undef LOCTEXT_NAMESPACE
diff --git a/Engine/Plugins/Enterprise/DataprepEditor/Source/DataprepEditor/Private/Widgets/SAssetsPreviewWidget.h b/Engine/Plugins/Enterprise/DataprepEditor/Source/DataprepEditor/Private/Widgets/SAssetsPreviewWidget.h
new file mode 100644
--- /dev/null
+++ b/Engine/Plugins/Enterprise/DataprepEditor/Source/DataprepEditor/Private/Widgets/SAssetsPreviewWidget.h
@@ -0,0 +1,76 @@
+// Copyright 1998-2019 Epic Games, Inc. All Rights Reserved.
+
+#pragma once
+
+#include "CoreMinimal.h"
+
+#include "Types/SlateEnums.h"
+#include "Widgets/SCompoundWidget.h"
+#include "Widgets/SWidget.h"
+#include "Widgets/Views/STableRow.h"
+#include "Widgets/Views/STableViewBase.h"
+#include "Widgets/Views/STreeView.h"
+
+
+namespace AssetPreviewWidget
+{
+	DECLARE_MULTICAST_DELEGATE_OneParam(FOnSelectionChanged, TSet< UObject* > /** Selected objects */)
+
+	using FAssetTreeItemPtr = TSharedPtr< struct FAssetTreeItem >;
+	using FAssetTreeItemWeakPtr = TWeakPtr< struct FAssetTreeItem >;
+
+	class SAssetsPreviewWidget : public SCompoundWidget
+	{
+	public:
+		SLATE_BEGIN_ARGS(SAssetsPreviewWidget) {}
+		SLATE_END_ARGS()
+
+	public:
+		void Construct(const FArguments& InArgs);
+
+		void SetAssetsList(const TArray< TWeakObjectPtr< UObject > >& InAssetsList, const FString& InPathToReplace, const FString& InSubstitutePath);
+		void ClearAssetList();
+
+		FOnSelectionChanged& OnSelectionChanged() { return OnSelectionChangedDelegate; };
+	
+		FText OnGetHighlightText() const;
+
+		const TSharedPtr< STreeView< FAssetTreeItemPtr > > GetTreeView() const
+		{
+			return TreeView;
+		}
+
+	private:
+
+		friend struct FAssetTreeItem;
+
+		void FilterAssetsNames();
+		void ExpandAllFolders();
+		void ExpandFolderRecursive(FAssetTreeItemPtr InItem);
+
+		void SortRecursive(TArray< FAssetTreeItemPtr >& InItems);
+
+		TArray< FString > GetItemsName(TWeakObjectPtr< UObject > Asset) const;
+
+		TSharedRef< class ITableRow > MakeRowWidget(FAssetTreeItemPtr InItem, const TSharedRef< class STableViewBase >& OwnerTable) const;
+		void OnGetChildren(FAssetTreeItemPtr InParent, TArray< FAssetTreeItemPtr >& OutChildren) const;
+
+		void OnSearchBoxChanged(const FText& InSearchText);
+		void OnSearchBoxCommitted(const FText& InSearchText, ETextCommit::Type CommitInfo);
+
+		void OnSetExpansionRecursive(FAssetTreeItemPtr InTreeNode, bool bInIsItemExpanded);
+
+		void OnSelectionChangedInternal(FAssetTreeItemPtr ItemSelected, ESelectInfo::Type SelectionType);
+
+		TArray< FAssetTreeItemPtr > RootItems;
+		TArray< FAssetTreeItemPtr > FilteredRootItems;
+
+		TSharedPtr< STreeView< FAssetTreeItemPtr > > TreeView;
+
+		FText FilterText;
+		FString PathToReplace;
+		FString SubstitutePath;
+
+		FOnSelectionChanged OnSelectionChangedDelegate;
+	};
+}
diff --git a/Engine/Plugins/Enterprise/DataprepEditor/Source/DataprepEditor/Private/Widgets/SDataprepActionMenu.cpp b/Engine/Plugins/Enterprise/DataprepEditor/Source/DataprepEditor/Private/Widgets/SDataprepActionMenu.cpp
new file mode 100644
--- /dev/null
+++ b/Engine/Plugins/Enterprise/DataprepEditor/Source/DataprepEditor/Private/Widgets/SDataprepActionMenu.cpp
@@ -0,0 +1,145 @@
+// Copyright 1998-2019 Epic Games, Inc. All Rights Reserved.
+
+#include "Widgets/SDataprepActionMenu.h"
+
+#include "BlueprintNodes/K2Node_DataprepAction.h"
+#include "DataprepActionAsset.h"
+#include "SGraphActionMenu.h"
+#include "SchemaActions/DataprepSchemaAction.h"
+#include "SchemaActions/DataprepSchemaActionUtils.h"
+#include "SchemaActions/IDataprepMenuActionCollector.h"
+
+#include "EdGraph/EdGraph.h"
+#include "EdGraph/EdGraphPin.h"
+#include "EdGraph/EdGraphSchema.h"
+#include "EditorStyleSet.h"
+#include "Framework/Application/SlateApplication.h"
+#include "ScopedTransaction.h"
+#include "Widgets/Views/SExpanderArrow.h"
+
+void SDataprepActionMenu::Construct(const FArguments& InArgs, TUniquePtr<IDataprepMenuActionCollector> InMenuActionCollector)
+{
+	MenuActionCollector = MoveTemp( InMenuActionCollector );
+	
+	Context = InArgs._DataprepActionContext;
+	TransactionTextGetter = InArgs._TransactionText;
+	GraphObj = InArgs._GraphObj;
+	NewNodePosition = InArgs._NewNodePosition;
+	DraggedFromPins = InArgs._DraggedFromPins;
+	OnClosedCallback = InArgs._OnClosedCallback;
+
+	ChildSlot
+	[
+		SNew( SBorder )
+		.BorderImage( FEditorStyle::GetBrush("Menu.Background") )
+		.Padding( 5 )
+		[
+			SNew( SBox )
+			.WidthOverride( 300 )
+			.HeightOverride( 400 )
+			[
+				SAssignNew( ActionMenu, SGraphActionMenu )
+				.OnActionSelected( this, &SDataprepActionMenu::OnActionSelected )
+				.OnCollectAllActions( this, &SDataprepActionMenu::CollectActions )
+				.OnCreateCustomRowExpander( this, &SDataprepActionMenu::OnCreateCustomRowExpander )
+				.AutoExpandActionMenu( MenuActionCollector->ShouldAutoExpand() )
+				.ShowFilterTextBox( true )
+			]
+		]
+	];
+}
+
+TSharedPtr<class SEditableTextBox> SDataprepActionMenu::GetFilterTextBox()
+{
+	return ActionMenu->GetFilterTextBox();
+}
+
+SDataprepActionMenu::~SDataprepActionMenu()
+{
+	OnClosedCallback.ExecuteIfBound();
+}
+
+void SDataprepActionMenu::CollectActions(FGraphActionListBuilderBase& OutActions)
+{
+	for ( TSharedPtr<FDataprepSchemaAction> Action : MenuActionCollector->CollectActions() )
+	{
+		OutActions.AddAction( StaticCastSharedPtr<FEdGraphSchemaAction>( Action ) );
+	}
+}
+
+void SDataprepActionMenu::OnActionSelected(const TArray< TSharedPtr<FEdGraphSchemaAction> >& SelectedActions, ESelectInfo::Type InSelectionType)
+{
+	/**
+	 * Create a node if needed and execute the action
+	 * return true if the action was properly executed
+	 */
+	auto ExcuteAction = [this](FDataprepSchemaAction& DataprepAction) -> bool
+		{
+			UK2Node_DataprepAction* ActionNode = nullptr;
+
+			if ( ShouldCreateNewNode() )
+			{
+				ActionNode = DataprepSchemaActionUtils::SpawnEdGraphNode<UK2Node_DataprepAction>( *GraphObj, NewNodePosition );
+				ActionNode->CreateDataprepActionAsset();
+				Context.DataprepActionPtr = MakeWeakObjectPtr<UDataprepActionAsset>( ActionNode->GetDataprepAction() );
+
+				if ( DraggedFromPins.Num() > 0 )
+				{
+					ActionNode->AutowireNewNode( DraggedFromPins[0] );
+				}
+			}
+
+			if ( !DataprepAction.ExecuteAction( Context ) )
+			{
+				if ( ActionNode )
+				{
+					// Remove the created node
+					GraphObj->RemoveNode( ActionNode );
+				}
+				return false;
+			}
+
+			return true;
+		};
+
+	for ( const TSharedPtr<FEdGraphSchemaAction>& Action : SelectedActions )
+	{
+		FDataprepSchemaAction* DataprepAction = static_cast< FDataprepSchemaAction* >( Action.Get() );
+		if ( DataprepAction )
+		{
+			if ( TransactionTextGetter.IsSet() )
+			{
+				FScopedTransaction Transaction( TransactionTextGetter.Get() );
+
+				if ( !ExcuteAction( *DataprepAction ) )
+				{
+					Transaction.Cancel();
+				}
+			}
+			else
+			{
+				ExcuteAction( *DataprepAction );
+			}
+		}
+	}
+
+	if ( SelectedActions.Num() > 0 )
+	{
+		FSlateApplication::Get().DismissAllMenus();
+	}
+}
+
+TSharedRef<SExpanderArrow> SDataprepActionMenu::OnCreateCustomRowExpander(const FCustomExpanderData& InCustomExpanderData) const
+{
+	return SNew( SExpanderArrow, InCustomExpanderData.TableRow );
+}
+
+bool SDataprepActionMenu::ShouldCreateNewNode() const
+{
+	if ( GraphObj && !Context.DataprepActionPtr.Get() )
+	{
+		return true;
+	}
+
+	return false;
+}
diff --git a/Engine/Plugins/Enterprise/DataprepEditor/Source/DataprepLibraries/Private/DataprepLibrariesModule.cpp b/Engine/Plugins/Enterprise/DataprepEditor/Source/DataprepLibraries/Private/DataprepLibrariesModule.cpp
new file mode 100644
--- /dev/null
+++ b/Engine/Plugins/Enterprise/DataprepEditor/Source/DataprepLibraries/Private/DataprepLibrariesModule.cpp
@@ -0,0 +1,26 @@
+// Copyright 1998-2019 Epic Games, Inc. All Rights Reserved.
+
+#include "DataprepLibrariesModule.h"
+#include "DataprepOperations.h"
+#include "DataprepOperationsLibrary.h"
+
+#include "Modules/ModuleManager.h"
+#include "PropertyEditorModule.h"
+
+class FDataprepLibrariesModule : public IDataprepLibrariesModule
+{
+public:
+	virtual void StartupModule() override
+	{
+		FPropertyEditorModule& PropertyModule = FModuleManager::LoadModuleChecked< FPropertyEditorModule >( TEXT("PropertyEditor") );
+		PropertyModule.RegisterCustomClassLayout( UDataprepSetLODGroupOperation::StaticClass()->GetFName(), FOnGetDetailCustomizationInstance::CreateStatic( &FDataprepSetLOGGroupDetails::MakeDetails ) );
+	}
+
+	virtual void ShutdownModule() override
+	{
+		FPropertyEditorModule& PropertyModule = FModuleManager::LoadModuleChecked< FPropertyEditorModule >( TEXT("PropertyEditor") );
+		PropertyModule.UnregisterCustomClassLayout( TEXT("DataprepSetLODGroupOperation") );
+	}
+};
+
+IMPLEMENT_MODULE( FDataprepLibrariesModule, DataprepLibraries )
diff --git a/Engine/Plugins/Enterprise/DataprepEditor/Source/DataprepLibraries/Private/DataprepOperationsLibrary.cpp b/Engine/Plugins/Enterprise/DataprepEditor/Source/DataprepLibraries/Private/DataprepOperationsLibrary.cpp
new file mode 100644
--- /dev/null
+++ b/Engine/Plugins/Enterprise/DataprepEditor/Source/DataprepLibraries/Private/DataprepOperationsLibrary.cpp
@@ -0,0 +1,681 @@
+// Copyright 1998-2019 Epic Games, Inc. All Rights Reserved.
+
+#include "DataprepOperationsLibrary.h"
+
+#include "DataprepCoreUtils.h"
+
+#include "ActorEditorUtils.h"
+#include "AssetDeleteModel.h"
+#include "AssetDeleteModel.h"
+#include "AssetRegistryModule.h"
+#include "Camera/CameraActor.h"
+#include "Editor.h"
+#include "Engine/Light.h"
+#include "Engine/StaticMesh.h"
+#include "Engine/StaticMeshActor.h"
+#include "Engine/Texture.h"
+#include "Engine/World.h"
+#include "EngineUtils.h"
+#include "GameFramework/Actor.h"
+#include "IMeshBuilderModule.h"
+#include "Interfaces/ITargetPlatform.h"
+#include "Interfaces/ITargetPlatformManagerModule.h"
+#include "Materials/Material.h"
+#include "Materials/MaterialInterface.h"
+#include "Math/Vector2D.h"
+#include "MeshDescriptionOperations.h"
+#include "Misc/FileHelper.h"
+#include "ObjectTools.h"
+#include "StaticMeshAttributes.h"
+#include "TessellationRendering.h"
+#include "UObject/SoftObjectPath.h"
+
+DEFINE_LOG_CATEGORY(LogDataprep);
+
+#define LOCTEXT_NAMESPACE "DataprepOperationsLibrary"
+
+extern UNREALED_API UEditorEngine* GEditor;
+
+namespace DataprepOperationsLibraryUtil
+{
+	TSet<UStaticMesh*> GetSelectedMeshes(const TArray<AActor*>& SelectedActors)
+	{
+		TSet<UStaticMesh*> SelectedMeshes;
+
+		for (AActor* Actor : SelectedActors)
+		{
+			if (Actor)
+			{
+				TInlineComponentArray<UStaticMeshComponent*> StaticMeshComponents(Actor);
+				for (UStaticMeshComponent* StaticMeshComponent : StaticMeshComponents)
+				{
+					if(UStaticMesh* StaticMesh = StaticMeshComponent->GetStaticMesh())
+					{
+						SelectedMeshes.Add( StaticMesh );
+					}
+				}
+			}
+		}
+
+		return SelectedMeshes;
+	}
+
+	TSet<UStaticMesh*> GetSelectedMeshes(const TArray<UObject*>& SelectedObjects)
+	{
+		TSet<UStaticMesh*> SelectedMeshes;
+		TArray<AActor*> SelectedActors;
+
+		for (UObject* Object : SelectedObjects)
+		{
+			if ( UStaticMesh* StaticMesh = Cast<UStaticMesh>(Object) )
+			{
+				SelectedMeshes.Add( StaticMesh );
+			}
+			else if ( Object->IsA(UStaticMeshComponent::StaticClass()) )
+			{
+				if ((StaticMesh = Cast<UStaticMeshComponent>(Object)->GetStaticMesh()) != nullptr )
+				{
+					SelectedMeshes.Add(StaticMesh);
+				}
+			}
+			else if (AActor* Actor = Cast<AActor>(Object) )
+			{
+				TInlineComponentArray<UStaticMeshComponent*> StaticMeshComponents( Actor );
+				for (UStaticMeshComponent* StaticMeshComponent : StaticMeshComponents)
+				{
+					if((StaticMesh = StaticMeshComponent->GetStaticMesh()) != nullptr)
+					{
+						SelectedMeshes.Add( StaticMesh );
+					}
+				}
+			}
+		}
+
+		return SelectedMeshes;
+	}
+
+	TArray<UMaterialInterface*> GetUsedMaterials(const TArray<UObject*>& SelectedObjects)
+	{
+		TSet<UMaterialInterface*> MaterialSet;
+
+		for (UObject* Object : SelectedObjects)
+		{
+			if (AActor* Actor = Cast< AActor >(Object))
+			{
+				// Find the materials by iterating over every mesh component.
+				TInlineComponentArray<UMeshComponent*> MeshComponents(Actor);
+				for (UMeshComponent* MeshComponent : MeshComponents)
+				{
+					int32 MaterialCount = FMath::Max( MeshComponent->GetNumOverrideMaterials(), MeshComponent->GetNumMaterials() );
+
+					for (int32 Index = 0; Index < MaterialCount; ++Index)
+					{
+						MaterialSet.Add(MeshComponent->GetMaterial(Index));
+					}
+				}
+			}
+			else if (UStaticMesh* StaticMesh = Cast< UStaticMesh >(Object))
+			{
+				for (int32 Index = 0; Index < StaticMesh->StaticMaterials.Num(); ++Index)
+				{
+					MaterialSet.Add(StaticMesh->GetMaterial(Index));
+				}
+			}
+		}
+
+		return MaterialSet.Array();
+	}
+
+	TArray<UStaticMesh*> GetUsedMeshes(const TArray<UObject*>& SelectedObjects)
+	{
+		TSet<UStaticMesh*> MeshesSet;
+
+		for (UObject* Object : SelectedObjects)
+		{
+			if (AActor* Actor = Cast< AActor >(Object))
+			{
+				// Find the meshes by iterating over every mesh component.
+				TInlineComponentArray<UStaticMeshComponent*> MeshComponents(Actor);
+				for (UStaticMeshComponent* MeshComponent : MeshComponents)
+				{
+					if(MeshComponent && MeshComponent->GetStaticMesh())
+					{
+						MeshesSet.Add( MeshComponent->GetStaticMesh() );
+					}
+				}
+			}
+		}
+
+		return MeshesSet.Array();
+	}
+
+	class FScopedStaticMeshEdit final
+	{
+	public:
+		FScopedStaticMeshEdit( UStaticMesh* InStaticMesh )
+			: StaticMesh( InStaticMesh )
+		{
+			BuildSettingsBackup = PreventStaticMeshBuild( StaticMesh );
+		}
+
+		FScopedStaticMeshEdit( const FScopedStaticMeshEdit& Other ) = default;
+		FScopedStaticMeshEdit( FScopedStaticMeshEdit&& Other ) = default;
+
+		FScopedStaticMeshEdit& operator=( const FScopedStaticMeshEdit& Other ) = default;
+		FScopedStaticMeshEdit& operator=( FScopedStaticMeshEdit&& Other ) = default;
+
+		~FScopedStaticMeshEdit()
+		{
+			RestoreStaticMeshBuild( StaticMesh, MoveTemp( BuildSettingsBackup ) );
+		}
+
+	public:
+		static TArray< FMeshBuildSettings > PreventStaticMeshBuild( UStaticMesh* StaticMesh )
+		{
+			if ( !StaticMesh )
+			{
+				return {};
+			}
+
+			TArray< FMeshBuildSettings > BuildSettingsBackup;
+
+			for ( FStaticMeshSourceModel& SourceModel : StaticMesh->GetSourceModels() )
+			{
+				BuildSettingsBackup.Add( SourceModel.BuildSettings );
+
+				// These were done in the PreBuild step
+				SourceModel.BuildSettings.bGenerateLightmapUVs = false;
+				SourceModel.BuildSettings.bRecomputeNormals = false;
+				SourceModel.BuildSettings.bRecomputeTangents = false;
+				SourceModel.BuildSettings.bBuildAdjacencyBuffer = false;
+				SourceModel.BuildSettings.bBuildReversedIndexBuffer = false;
+				SourceModel.BuildSettings.bComputeWeightedNormals = false;
+			}
+
+			return BuildSettingsBackup;
+		}
+
+		static void RestoreStaticMeshBuild( UStaticMesh* StaticMesh, const TArray< FMeshBuildSettings >& BuildSettingsBackup )
+		{
+			if ( !StaticMesh )
+			{
+				return;
+			}
+
+			// Restore StaticMesh's build settings
+			for ( int32 LODIndex = 0; LODIndex < BuildSettingsBackup.Num() ; ++LODIndex )
+			{
+				// Update only LODs which were cached
+				if (StaticMesh->IsSourceModelValid( LODIndex ))
+				{
+					const FMeshBuildSettings& CachedBuildSettings = BuildSettingsBackup[ LODIndex ];
+					FMeshBuildSettings& BuildSettings = StaticMesh->GetSourceModel(LODIndex).BuildSettings;
+
+					// Restore only the properties which were modified
+					BuildSettings.bGenerateLightmapUVs = CachedBuildSettings.bGenerateLightmapUVs;
+					BuildSettings.bRecomputeNormals = CachedBuildSettings.bRecomputeNormals;
+					BuildSettings.bRecomputeTangents = CachedBuildSettings.bRecomputeTangents;
+					BuildSettings.bBuildAdjacencyBuffer = CachedBuildSettings.bBuildAdjacencyBuffer;
+					BuildSettings.bBuildReversedIndexBuffer = CachedBuildSettings.bBuildReversedIndexBuffer;
+					BuildSettings.bComputeWeightedNormals = CachedBuildSettings.bComputeWeightedNormals;
+				}
+			}
+		}
+
+	private:
+		TArray< FMeshBuildSettings > BuildSettingsBackup;
+		UStaticMesh* StaticMesh;
+	};
+
+	/** Customized version of UStaticMesh::SetMaterial avoiding the triggering of UStaticMesh::Build and its side-effects */
+	void SetMaterial( UStaticMesh* StaticMesh, int32 MaterialIndex, UMaterialInterface* NewMaterial )
+	{
+		if( StaticMesh->StaticMaterials.IsValidIndex( MaterialIndex ) )
+		{
+			FStaticMaterial& StaticMaterial = StaticMesh->StaticMaterials[ MaterialIndex ];
+			StaticMaterial.MaterialInterface = NewMaterial;
+			if( NewMaterial != nullptr )
+			{
+				if ( StaticMaterial.MaterialSlotName == NAME_None )
+				{
+					StaticMaterial.MaterialSlotName = NewMaterial->GetFName();
+				}
+
+				// Make sure adjacency information fit new material change
+				if( RequiresAdjacencyInformation( NewMaterial, nullptr, GWorld->FeatureLevel ) )
+				{
+					for( FStaticMeshSourceModel& SourceModel : StaticMesh->GetSourceModels() )
+					{
+						SourceModel.BuildSettings.bBuildAdjacencyBuffer = true;
+					}
+				}
+			}
+		}
+	}
+
+	FStaticMeshBuilder::FStaticMeshBuilder(const TSet<UStaticMesh *>& InStaticMeshes)
+	{
+		StaticMeshes = BuildStaticMeshes( InStaticMeshes );
+	}
+
+	FStaticMeshBuilder::~FStaticMeshBuilder()
+	{
+		// Release render data of built static meshes
+		for(UStaticMesh* StaticMesh : StaticMeshes)
+		{
+			if(StaticMesh)
+			{
+				StaticMesh->RenderData.Reset();
+			}
+		}
+	}
+
+	TArray<UStaticMesh*> BuildStaticMeshes(const TSet<UStaticMesh*>& StaticMeshes, bool bForceBuild)
+	{
+		TRACE_CPUPROFILER_EVENT_SCOPE(DataprepOperationsLibraryUtil::BuildStaticMeshes);
+
+		TArray<UStaticMesh*> BuiltMeshes;
+		BuiltMeshes.Reserve( StaticMeshes.Num() );
+
+		if(bForceBuild)
+		{
+			BuiltMeshes.Append( StaticMeshes.Array() );
+		}
+		else
+		{
+			for(UStaticMesh* StaticMesh : StaticMeshes)
+			{
+				if(StaticMesh && (!StaticMesh->RenderData.IsValid() || !StaticMesh->RenderData->IsInitialized()))
+				{
+					BuiltMeshes.Add( StaticMesh );
+				}
+			}
+		}
+
+		if(BuiltMeshes.Num() > 0)
+		{
+			// Start with the biggest mesh first to help balancing tasks on threads
+			BuiltMeshes.Sort(
+				[](const UStaticMesh& Lhs, const UStaticMesh& Rhs) 
+			{ 
+				int32 LhsVerticesNum = Lhs.IsMeshDescriptionValid(0) ? Lhs.GetMeshDescription(0)->Vertices().Num() : 0;
+				int32 RhsVerticesNum = Rhs.IsMeshDescriptionValid(0) ? Rhs.GetMeshDescription(0)->Vertices().Num() : 0;
+
+				return LhsVerticesNum > RhsVerticesNum;
+			}
+			);
+
+			//Cache the BuildSettings and update them before building the meshes.
+			TArray< TArray<FMeshBuildSettings> > StaticMeshesSettings;
+			StaticMeshesSettings.Reserve( BuiltMeshes.Num() );
+
+			for (UStaticMesh* StaticMesh : BuiltMeshes)
+			{
+				TArray<FStaticMeshSourceModel>& SourceModels = StaticMesh->GetSourceModels();
+				TArray<FMeshBuildSettings> BuildSettings;
+				BuildSettings.Reserve(SourceModels.Num());
+
+				for(int32 Index = 0; Index < SourceModels.Num(); ++Index)
+				{
+					FStaticMeshSourceModel& SourceModel = SourceModels[Index];
+
+					BuildSettings.Add( SourceModel.BuildSettings );
+
+					if(FMeshDescription* MeshDescription = StaticMesh->GetMeshDescription(Index))
+					{
+						FStaticMeshAttributes Attributes(*MeshDescription);
+						if(SourceModel.BuildSettings.DstLightmapIndex != -1)
+						{
+							TVertexInstanceAttributesConstRef<FVector2D> VertexInstanceUVs = Attributes.GetVertexInstanceUVs();
+							SourceModel.BuildSettings.bGenerateLightmapUVs = VertexInstanceUVs.IsValid() && VertexInstanceUVs.GetNumIndices() > SourceModel.BuildSettings.DstLightmapIndex;
+						}
+						else
+						{
+							SourceModel.BuildSettings.bGenerateLightmapUVs = false;
+						}
+
+						SourceModel.BuildSettings.bRecomputeNormals = !(Attributes.GetVertexInstanceNormals().IsValid() && Attributes.GetVertexInstanceNormals().GetNumIndices() > 0);
+						SourceModel.BuildSettings.bRecomputeTangents = false;
+						//SourceModel.BuildSettings.bBuildAdjacencyBuffer = false;
+						//SourceModel.BuildSettings.bBuildReversedIndexBuffer = false;
+					}
+				}
+
+				StaticMeshesSettings.Add(MoveTemp(BuildSettings));				
+			}
+
+			// Disable warnings from LogStaticMesh. Not useful
+			ELogVerbosity::Type PrevLogStaticMeshVerbosity = LogStaticMesh.GetVerbosity();
+			LogStaticMesh.SetVerbosity( ELogVerbosity::Error );
+
+			UStaticMesh::BatchBuild(BuiltMeshes, true );
+
+			// Restore LogStaticMesh verbosity
+			LogStaticMesh.SetVerbosity( PrevLogStaticMeshVerbosity );
+
+			for(int32 Index = 0; Index < BuiltMeshes.Num(); ++Index)
+			{
+				UStaticMesh* StaticMesh = BuiltMeshes[Index];
+				TArray<FMeshBuildSettings>& PrevBuildSettings = StaticMeshesSettings[Index];
+
+				TArray<FStaticMeshSourceModel>& SourceModels = StaticMesh->GetSourceModels();
+
+				for(int32 SourceModelIndex = 0; SourceModelIndex < SourceModels.Num(); ++SourceModelIndex)
+				{
+					SourceModels[SourceModelIndex].BuildSettings = PrevBuildSettings[SourceModelIndex];
+				}
+
+				for ( FStaticMeshLODResources& LODResources : StaticMesh->RenderData->LODResources )
+				{
+					LODResources.bHasColorVertexData = true;
+				}
+			}
+		}
+
+		return BuiltMeshes;
+	}
+} // ns DataprepOperationsLibraryUtil
+
+void UDataprepOperationsLibrary::SetLods(const TArray<UObject*>& SelectedObjects, const FEditorScriptingMeshReductionOptions& ReductionOptions, TArray<UObject*>& ModifiedObjects)
+{
+	TSet<UStaticMesh*> SelectedMeshes = DataprepOperationsLibraryUtil::GetSelectedMeshes(SelectedObjects);
+
+	// Create LODs but do not commit changes
+	for (UStaticMesh* StaticMesh : SelectedMeshes)
+	{
+		if (StaticMesh)
+		{
+			DataprepOperationsLibraryUtil::FScopedStaticMeshEdit StaticMeshEdit( StaticMesh );
+
+			UEditorStaticMeshLibrary::SetLodsWithNotification(StaticMesh, ReductionOptions, false);
+
+			ModifiedObjects.Add( StaticMesh );
+		}
+	}
+}
+
+void UDataprepOperationsLibrary::SetSimpleCollision(const TArray<UObject*>& SelectedObjects, const EScriptingCollisionShapeType ShapeType, TArray<UObject*>& ModifiedObjects)
+{
+	TSet<UStaticMesh*> SelectedMeshes = DataprepOperationsLibraryUtil::GetSelectedMeshes(SelectedObjects);
+
+	// Make sure all static meshes to be processed have render data for NDOP types
+	bool bNeedRenderData = false;
+	switch (ShapeType)
+	{
+		case EScriptingCollisionShapeType::NDOP10_X:
+		case EScriptingCollisionShapeType::NDOP10_Y:
+		case EScriptingCollisionShapeType::NDOP10_Z:
+		case EScriptingCollisionShapeType::NDOP18:
+		case EScriptingCollisionShapeType::NDOP26:
+		{
+			bNeedRenderData = true;
+			break;
+		}
+		default:
+		{
+			break;
+		}
+	}
+
+	DataprepOperationsLibraryUtil::FStaticMeshBuilder StaticMeshBuilder( bNeedRenderData ? SelectedMeshes : TSet<UStaticMesh*>() );
+
+	// Create LODs but do not commit changes
+	for (UStaticMesh* StaticMesh : SelectedMeshes)
+	{
+		if (StaticMesh)
+		{
+			DataprepOperationsLibraryUtil::FScopedStaticMeshEdit StaticMeshEdit( StaticMesh );
+
+			// Remove existing simple collisions
+			UEditorStaticMeshLibrary::RemoveCollisionsWithNotification( StaticMesh, false );
+
+			UEditorStaticMeshLibrary::AddSimpleCollisionsWithNotification( StaticMesh, ShapeType, false );
+
+			ModifiedObjects.Add( StaticMesh );
+		}
+	}
+}
+
+void UDataprepOperationsLibrary::SetConvexDecompositionCollision(const TArray<UObject*>& SelectedObjects, int32 HullCount, int32 MaxHullVerts, int32 HullPrecision, TArray<UObject*>& ModifiedObjects)
+{
+	TSet<UStaticMesh*> SelectedMeshes = DataprepOperationsLibraryUtil::GetSelectedMeshes(SelectedObjects);
+
+	// Make sure all static meshes to be processed have render data
+	DataprepOperationsLibraryUtil::FStaticMeshBuilder StaticMeshBuilder(SelectedMeshes);
+
+	// Build complex collision
+	for (UStaticMesh* StaticMesh : SelectedMeshes)
+	{
+		if (StaticMesh)
+		{
+			DataprepOperationsLibraryUtil::FScopedStaticMeshEdit StaticMeshEdit( StaticMesh );
+
+			UEditorStaticMeshLibrary::SetConvexDecompositionCollisionsWithNotification(StaticMesh, HullCount, MaxHullVerts, HullPrecision, false);
+
+			ModifiedObjects.Add( StaticMesh );
+		}
+	}
+}
+
+void UDataprepOperationsLibrary::SubstituteMaterial(const TArray<UObject*>& SelectedObjects, const FString& MaterialSearch, EEditorScriptingStringMatchType StringMatch, UMaterialInterface* MaterialSubstitute)
+{
+	TArray<UMaterialInterface*> MaterialsUsed = DataprepOperationsLibraryUtil::GetUsedMaterials(SelectedObjects);
+
+	SubstituteMaterial(SelectedObjects, MaterialSearch, StringMatch, MaterialsUsed, MaterialSubstitute);
+}
+
+void UDataprepOperationsLibrary::SubstituteMaterialsByTable(const TArray<UObject*>& SelectedObjects, const UDataTable* DataTable)
+{
+	if (DataTable == nullptr || DataTable->GetRowStruct() == nullptr || !DataTable->GetRowStruct()->IsChildOf(FMaterialSubstitutionDataTable::StaticStruct()))
+	{
+		return;
+	}
+
+	TArray<UMaterialInterface*> MaterialsUsed = DataprepOperationsLibraryUtil::GetUsedMaterials(SelectedObjects);
+
+	const TMap<FName, uint8*>&  MaterialTableRowMap = DataTable->GetRowMap();
+	for (auto& MaterialTableRowEntry : MaterialTableRowMap)
+	{
+		const FMaterialSubstitutionDataTable* MaterialRow = (const FMaterialSubstitutionDataTable*)MaterialTableRowEntry.Value;
+		if (MaterialRow != nullptr && MaterialRow->MaterialReplacement != nullptr)
+		{
+			SubstituteMaterial(SelectedObjects, MaterialRow->SearchString, MaterialRow->StringMatch, MaterialsUsed, MaterialRow->MaterialReplacement);
+		}
+	}
+}
+
+void UDataprepOperationsLibrary::SubstituteMaterial(const TArray<UObject*>& SelectedObjects, const FString& MaterialSearch, EEditorScriptingStringMatchType StringMatch, const TArray<UMaterialInterface*>& MaterialList, UMaterialInterface* MaterialSubstitute)
+{
+	TArray<UObject*> MatchingObjects = UEditorFilterLibrary::ByIDName(TArray<UObject*>(MaterialList), MaterialSearch, StringMatch, EEditorScriptingFilterType::Include);
+
+	TArray<UMaterialInterface*> MaterialsToReplace;
+	for (UObject* Object : MatchingObjects)
+	{
+		if (UMaterialInterface* MaterialInterface = Cast<UMaterialInterface>(Object))
+		{
+			MaterialsToReplace.Add(MaterialInterface);
+		}
+	}
+
+	for (UMaterialInterface* MaterialToReplace : MaterialsToReplace)
+	{
+		for (UObject* Object : SelectedObjects)
+		{
+			if (AActor* Actor = Cast< AActor >(Object))
+			{
+				// Find the materials by iterating over every mesh component.
+				TInlineComponentArray<UMeshComponent*> MeshComponents(Actor);
+				for (UMeshComponent* MeshComponent : MeshComponents)
+				{
+					int32 MaterialCount = FMath::Max( MeshComponent->GetNumOverrideMaterials(), MeshComponent->GetNumMaterials() );
+
+					for (int32 Index = 0; Index < MaterialCount; ++Index)
+					{
+						if (MeshComponent->GetMaterial(Index) == MaterialToReplace)
+						{
+							MeshComponent->SetMaterial(Index, MaterialSubstitute);
+						}
+					}
+				}
+			}
+			else if (UStaticMesh* StaticMesh = Cast< UStaticMesh >(Object))
+			{
+				DataprepOperationsLibraryUtil::FScopedStaticMeshEdit StaticMeshEdit( StaticMesh );
+
+				TArray<FStaticMaterial>& StaticMaterials = StaticMesh->StaticMaterials;
+				for (int32 Index = 0; Index < StaticMesh->StaticMaterials.Num(); ++Index)
+				{
+					if (StaticMesh->GetMaterial(Index) == MaterialToReplace)
+					{
+						DataprepOperationsLibraryUtil::SetMaterial( StaticMesh, Index, MaterialSubstitute );
+					}
+				}
+			}
+		}
+	}
+}
+
+void UDataprepOperationsLibrary::SetMobility( const TArray< UObject* >& SelectedObjects, EComponentMobility::Type MobilityType )
+{
+	for (UObject* Object : SelectedObjects)
+	{
+		if (AActor* Actor = Cast< AActor >(Object))
+		{
+			// Find the materials by iterating over every mesh component.
+			TInlineComponentArray<USceneComponent*> SceneComponents(Actor);
+			for (USceneComponent* SceneComponent : SceneComponents)
+			{
+				SceneComponent->SetMobility(MobilityType);
+			}
+		}
+	}
+}
+
+void UDataprepOperationsLibrary::SetMaterial( const TArray< UObject* >& SelectedObjects, UMaterialInterface* MaterialSubstitute )
+{
+	for (UObject* Object : SelectedObjects)
+	{
+		if (AActor* Actor = Cast< AActor >(Object))
+		{
+			// Find the materials by iterating over every mesh component.
+			TInlineComponentArray<UMeshComponent*> MeshComponents(Actor);
+			for (UMeshComponent* MeshComponent : MeshComponents)
+			{
+				int32 MaterialCount = FMath::Max( MeshComponent->GetNumOverrideMaterials(), MeshComponent->GetNumMaterials() );
+
+				for (int32 Index = 0; Index < MaterialCount; ++Index)
+				{
+					MeshComponent->SetMaterial(Index, MaterialSubstitute);
+				}
+			}
+		}
+		else if (UStaticMesh* StaticMesh = Cast< UStaticMesh >(Object))
+		{
+			DataprepOperationsLibraryUtil::FScopedStaticMeshEdit StaticMeshEdit( StaticMesh );
+
+			for (int32 Index = 0; Index < StaticMesh->StaticMaterials.Num(); ++Index)
+			{
+				DataprepOperationsLibraryUtil::SetMaterial( StaticMesh, Index, MaterialSubstitute );
+			}
+		}
+	}
+}
+
+void UDataprepOperationsLibrary::SetLODGroup( const TArray<UObject*>& SelectedObjects, FName& LODGroupName, TArray<UObject*>& ModifiedObjects )
+{
+	TArray<FName> LODGroupNames;
+	UStaticMesh::GetLODGroups( LODGroupNames );
+
+	if ( LODGroupNames.Find( LODGroupName ) != INDEX_NONE )
+	{
+		TSet<UStaticMesh*> SelectedMeshes = DataprepOperationsLibraryUtil::GetSelectedMeshes(SelectedObjects);
+
+		// Apply the new LODGroup without rebuilding the static mesh
+		for (UStaticMesh* StaticMesh : SelectedMeshes)
+		{
+			if(StaticMesh)
+			{
+				StaticMesh->SetLODGroup( LODGroupName, false);
+				ModifiedObjects.Add( StaticMesh );
+			}
+		}
+	}
+}
+
+void UDataprepOperationsLibrary::SetMesh(const TArray<UObject*>& SelectedObjects, UStaticMesh* MeshSubstitute)
+{
+	for (UObject* Object : SelectedObjects)
+	{
+		if (AActor* Actor = Cast< AActor >(Object))
+		{
+			// Find the meshes by iterating over every mesh component.
+			TInlineComponentArray<UStaticMeshComponent*> MeshComponents(Actor);
+			for (UStaticMeshComponent* MeshComponent : MeshComponents)
+			{
+				if(MeshComponent)
+				{
+					MeshComponent->SetStaticMesh( MeshSubstitute );
+				}
+			}
+		}
+	}
+}
+
+void UDataprepOperationsLibrary::SubstituteMesh(const TArray<UObject*>& SelectedObjects, const FString& MeshSearch, EEditorScriptingStringMatchType StringMatch, UStaticMesh* MeshSubstitute)
+{
+	TArray<UStaticMesh*> MeshesUsed = DataprepOperationsLibraryUtil::GetUsedMeshes(SelectedObjects);
+
+	SubstituteMesh( SelectedObjects, MeshSearch, StringMatch, MeshesUsed, MeshSubstitute );
+}
+
+void UDataprepOperationsLibrary::SubstituteMeshesByTable(const TArray<UObject*>& SelectedObjects, const UDataTable * DataTable)
+{
+	if (DataTable == nullptr || DataTable->GetRowStruct() == nullptr || !DataTable->GetRowStruct()->IsChildOf(FMeshSubstitutionDataTable::StaticStruct()))
+	{
+		return;
+	}
+
+	TArray<UStaticMesh*> MeshesUsed = DataprepOperationsLibraryUtil::GetUsedMeshes(SelectedObjects);
+
+	const TMap<FName, uint8*>&  MeshTableRowMap = DataTable->GetRowMap();
+	for (auto& MeshTableRowEntry : MeshTableRowMap)
+	{
+		const FMeshSubstitutionDataTable* MeshRow = (const FMeshSubstitutionDataTable*)MeshTableRowEntry.Value;
+		if (MeshRow != nullptr && MeshRow->MeshReplacement != nullptr)
+		{
+			SubstituteMesh( SelectedObjects, MeshRow->SearchString, MeshRow->StringMatch, MeshesUsed, MeshRow->MeshReplacement );
+		}
+	}
+}
+
+void UDataprepOperationsLibrary::SubstituteMesh(const TArray<UObject*>& SelectedObjects, const FString& MeshSearch, EEditorScriptingStringMatchType StringMatch, const TArray<UStaticMesh*>& MeshList, UStaticMesh* MeshSubstitute)
+{
+	TArray<UObject*> MatchingObjects = UEditorFilterLibrary::ByIDName(TArray<UObject*>(MeshList), MeshSearch, StringMatch, EEditorScriptingFilterType::Include);
+
+	TSet<UStaticMesh*> MeshesToReplace;
+	for (UObject* Object : MatchingObjects)
+	{
+		if (UStaticMesh* StaticMesh = Cast<UStaticMesh>(Object))
+		{
+			MeshesToReplace.Add(StaticMesh);
+		}
+	}
+
+	for (UObject* Object : SelectedObjects)
+	{
+		if (AActor* Actor = Cast< AActor >(Object))
+		{
+			// Find the meshes by iterating over every mesh component.
+			TInlineComponentArray<UStaticMeshComponent*> MeshComponents(Actor);
+			for (UStaticMeshComponent* MeshComponent : MeshComponents)
+			{
+				if( MeshesToReplace.Contains( MeshComponent->GetStaticMesh() ) )
+				{
+					MeshComponent->SetStaticMesh( MeshSubstitute );
+				}
+			}
+		}
+	}
+}
+
+#undef LOCTEXT_NAMESPACE
diff --git a/Engine/Plugins/Enterprise/DataprepEditor/Source/DataprepLibraries/Private/Fetchers/DataprepFloatFetcherLibrary.cpp b/Engine/Plugins/Enterprise/DataprepEditor/Source/DataprepLibraries/Private/Fetchers/DataprepFloatFetcherLibrary.cpp
new file mode 100644
--- /dev/null
+++ b/Engine/Plugins/Enterprise/DataprepEditor/Source/DataprepLibraries/Private/Fetchers/DataprepFloatFetcherLibrary.cpp
@@ -0,0 +1,82 @@
+// Copyright 1998-2019 Epic Games, Inc. All Rights Reserved.
+
+#include "Fetchers/DataprepFloatFetcherLibrary.h"
+
+#include "Components/StaticMeshComponent.h"
+#include "Engine/StaticMesh.h"
+#include "GameFramework/Actor.h"
+#include "Misc/Optional.h"
+#include "UObject/Object.h"
+
+#define LOCTEXT_NAMESPACE "DataprepFloatFetcherLibrary"
+
+/* UDataprepFloatVolumeFetcher methods+
+ *****************************************************************************/
+float UDataprepFloatBoundingVolumeFetcher::Fetch_Implementation(const UObject* Object, bool& bOutFetchSucceded) const
+{
+	if ( Object && !Object->IsPendingKill() )
+	{
+		auto GetStaticMeshBoundingBox = [](const UStaticMesh* StaticMesh) -> FBox
+			{
+				if ( !StaticMesh )
+				{
+					return {};
+				}
+
+				return StaticMesh->GetBoundingBox();
+			};
+
+		TOptional<float> Volume;
+		if ( const AActor* Actor = Cast<const AActor>( Object ) )
+		{
+			FBox ActorBox(ForceInit);
+
+			for ( const UActorComponent* ActorComponent : Actor->GetComponents() )
+			{
+				if ( const UPrimitiveComponent* PrimComp = Cast<const UPrimitiveComponent>( ActorComponent ) )
+				{
+					FBox ComponentBox(ForceInit);
+					if ( const UStaticMeshComponent* StaticMeshComponent = Cast<const UStaticMeshComponent>( ActorComponent ) )
+					{
+						ComponentBox = GetStaticMeshBoundingBox( StaticMeshComponent->GetStaticMesh() );
+						ComponentBox.TransformBy( PrimComp->GetComponentToWorld() );
+					}
+					if ( PrimComp->IsRegistered() )
+					{
+						ActorBox += ComponentBox.IsValid? ComponentBox : PrimComp->Bounds.GetBox();
+					}
+				}
+			}
+
+			if ( ActorBox.IsValid )
+			{
+				Volume = ActorBox.GetVolume();
+			}
+		}
+		else if ( const UStaticMesh* StaticMesh = Cast<const UStaticMesh>( Object ) )
+		{
+			Volume = GetStaticMeshBoundingBox( StaticMesh ).GetVolume();
+		}
+
+		if ( Volume.IsSet() )
+		{
+			bOutFetchSucceded = true;
+			return Volume.GetValue();
+		}
+	}
+
+	bOutFetchSucceded = false;
+	return {};
+}
+
+bool UDataprepFloatBoundingVolumeFetcher::IsThreadSafe() const
+{
+	return true;
+}
+
+FText UDataprepFloatBoundingVolumeFetcher::GetNodeDisplayFetcherName_Implementation() const
+{
+	return LOCTEXT("BoundingVolumeFilterTitle", "Bounding Volume");
+}
+
+#undef LOCTEXT_NAMESPACE
diff --git a/Engine/Plugins/Enterprise/DataprepEditor/Source/DataprepLibraries/Public/DataprepLibrariesModule.h b/Engine/Plugins/Enterprise/DataprepEditor/Source/DataprepLibraries/Public/DataprepLibrariesModule.h
new file mode 100644
--- /dev/null
+++ b/Engine/Plugins/Enterprise/DataprepEditor/Source/DataprepLibraries/Public/DataprepLibrariesModule.h
@@ -0,0 +1,11 @@
+// Copyright 1998-2019 Epic Games, Inc. All Rights Reserved.
+
+#pragma once
+
+#include "CoreMinimal.h"
+#include "Modules/ModuleInterface.h"
+
+class IDataprepLibrariesModule : public IModuleInterface
+{
+public:
+};
