Trying to fix the case issue on linux

diff --git a/Engine/Plugins/Enterprise/DataprepEditor/Source/DataprepLibraries/Public/Fetchers/DataprepFloatFetcherLibrary.h b/Engine/Plugins/Enterprise/DataprepEditor/Source/DataprepLibraries/Public/Fetchers/DataprepFloatFetcherLibrary.h
new file mode 100644
--- /dev/null
+++ b/Engine/Plugins/Enterprise/DataprepEditor/Source/DataprepLibraries/Public/Fetchers/DataprepFloatFetcherLibrary.h
@@ -0,0 +1,24 @@
+// Copyright 1998-2019 Epic Games, Inc. All Rights Reserved.
+
+#pragma once
+
+#include "SelectionSystem/DataprepFloatFetcher.h"
+
+#include "CoreMinimal.h"
+
+#include "DataprepFloatFetcherLibrary.generated.h"
+
+UCLASS(BlueprintType, NotBlueprintable, Meta = (DisplayName="Bounding Volume", ToolTip = "Return the bouding box volume of object.\nFor an actor bouding box only the components with a collision enabled will be use."))
+class UDataprepFloatBoundingVolumeFetcher final : public UDataprepFloatFetcher
+{
+	GENERATED_BODY()
+public:
+	//~ UDataprepFloatFetcher interface
+	virtual float Fetch_Implementation(const UObject* Object, bool& bOutFetchSucceded) const final;
+	//~ End of UDataprepFloatFetcher interface
+
+	//~ UDataprepFetcher interface
+	virtual bool IsThreadSafe() const final;
+	virtual FText GetNodeDisplayFetcherName_Implementation() const;
+	//~ End of UDataprepFetcher interface
+};
