Copy up CL 11160371 from //UE4/Dev-Enterprise to fix

diff --git a/Engine/Plugins/Enterprise/DataprepEditor/Source/DataprepEditor/Private/DataprepGraph/DataprepGraphActionNode.cpp b/Engine/Plugins/Enterprise/DataprepEditor/Source/DataprepEditor/Private/DataprepGraph/DataprepGraphActionNode.cpp
--- a/Engine/Plugins/Enterprise/DataprepEditor/Source/DataprepEditor/Private/DataprepGraph/DataprepGraphActionNode.cpp
+++ b/Engine/Plugins/Enterprise/DataprepEditor/Source/DataprepEditor/Private/DataprepGraph/DataprepGraphActionNode.cpp
@@ -35,10 +35,16 @@ UDataprepGraphActionNode::UDataprepGraphActionNode()
 
 void UDataprepGraphActionNode::Initialize(UDataprepActionAsset* InDataprepActionAsset, int32 InExecutionOrder)
 {
-	ensure(InDataprepActionAsset);
-	DataprepActionAsset = InDataprepActionAsset;
-	ActionTitle = DataprepActionAsset->GetLabel();
-	ExecutionOrder = InExecutionOrder;
+	if(InDataprepActionAsset)
+	{
+		DataprepActionAsset = InDataprepActionAsset;
+		ActionTitle = DataprepActionAsset->GetLabel();
+		ExecutionOrder = InExecutionOrder;
+	}
+	else
+	{
+		ensure(false);
+	}
 }
 
 FLinearColor UDataprepGraphActionNode::GetNodeTitleColor() const
diff --git a/Engine/Plugins/Enterprise/DataprepEditor/Source/DataprepEditor/Private/Widgets/DataprepGraph/SDataprepGraphTrackNode.cpp b/Engine/Plugins/Enterprise/DataprepEditor/Source/DataprepEditor/Private/Widgets/DataprepGraph/SDataprepGraphTrackNode.cpp
--- a/Engine/Plugins/Enterprise/DataprepEditor/Source/DataprepEditor/Private/Widgets/DataprepGraph/SDataprepGraphTrackNode.cpp
+++ b/Engine/Plugins/Enterprise/DataprepEditor/Source/DataprepEditor/Private/Widgets/DataprepGraph/SDataprepGraphTrackNode.cpp
@@ -754,7 +754,7 @@ float SDataprepGraphTrackWidget::ValidateNodeAbscissa(float InAbscissa)
 		AbscissaRange.X = ActionSlots[0]->OffsetAttr.Get(FVector2D::ZeroVector).Left;
 		AbscissaRange.Y = ActionSlots[ActionSlots.Num() - 1]->OffsetAttr.Get(FVector2D::ZeroVector).Left;
 
-		return InAbscissa < InAbscissa ? AbscissaRange.X : (InAbscissa > AbscissaRange.Y ? AbscissaRange.Y : InAbscissa);
+		return InAbscissa < AbscissaRange.X ? AbscissaRange.X : (InAbscissa > AbscissaRange.Y ? AbscissaRange.Y : InAbscissa);
 	}
 
 	return InAbscissa;
