Dataprep Merge and Proxy Mesh Operators leave empty mesh actors in the hierarchy

Integrated from Dev-Enterprise (CL 11459275)

diff --git a/Engine/Plugins/Enterprise/DataprepEditor/Source/DataprepLibraries/Private/DataprepEditingOperations.cpp b/Engine/Plugins/Enterprise/DataprepEditor/Source/DataprepLibraries/Private/DataprepEditingOperations.cpp
--- a/Engine/Plugins/Enterprise/DataprepEditor/Source/DataprepLibraries/Private/DataprepEditingOperations.cpp
+++ b/Engine/Plugins/Enterprise/DataprepEditor/Source/DataprepLibraries/Private/DataprepEditingOperations.cpp
@@ -202,6 +202,9 @@ void UDataprepMergeActorsOperation::OnExecution_Implementation(const FDataprepCo
 	TArray<UObject*> ObjectsToDelete;
 	ObjectsToDelete.Reserve( ComponentsToMerge.Num() + ActorsToMerge.Num() );
 
+	// Sort merged components to detach children first
+	ComponentsToMerge.Sort([](const USceneComponent &A, const USceneComponent &B) -> bool { return A.IsAttachedTo(&B); });
+
 	// Simple way to delete the actors: detach the merged components if it's safe to do so
 	for (UPrimitiveComponent* Component : ComponentsToMerge)
 	{
