Merging CL #11136518 from Dev-Enterprise

Fixed crash when deleting Dataprep asset

#rnx

diff --git a/Engine/Plugins/Enterprise/DataprepEditor/Source/DataprepEditor/Private/DataprepGraph/DataprepGraph.cpp b/Engine/Plugins/Enterprise/DataprepEditor/Source/DataprepEditor/Private/DataprepGraph/DataprepGraph.cpp
--- a/Engine/Plugins/Enterprise/DataprepEditor/Source/DataprepEditor/Private/DataprepGraph/DataprepGraph.cpp
+++ b/Engine/Plugins/Enterprise/DataprepEditor/Source/DataprepEditor/Private/DataprepGraph/DataprepGraph.cpp
@@ -24,7 +24,7 @@ void UDataprepGraph::Initialize(UDataprepAsset* InDataprepAsset)
 	DataprepAssetPtr = InDataprepAsset;
 
 	// Add recipe graph editor node which will be used as a start point to populate 
-	RecipeNode = TStrongObjectPtr<UDataprepGraphRecipeNode>( Cast<UDataprepGraphRecipeNode>(CreateNode(UDataprepGraphRecipeNode::StaticClass(), false)) );
+	UDataprepGraphRecipeNode* RecipeNode = Cast<UDataprepGraphRecipeNode>(CreateNode(UDataprepGraphRecipeNode::StaticClass(), false));
 	RecipeNode->SetEnabledState(ENodeEnabledState::Disabled, true);
 }
 
@@ -116,8 +116,6 @@ void FDataprepEditor::CreateGraphEditor()
 			.TitleBar(TitleBarWidget)
 			.GraphToEdit(DataprepGraph.Get())
 			.GraphEvents(Events);
-
-		DataprepGraph->SetEditor(GraphEditor);
 	}
 }
 
diff --git a/Engine/Plugins/Enterprise/DataprepEditor/Source/DataprepEditor/Private/DataprepGraph/DataprepGraph.h b/Engine/Plugins/Enterprise/DataprepEditor/Source/DataprepEditor/Private/DataprepGraph/DataprepGraph.h
--- a/Engine/Plugins/Enterprise/DataprepEditor/Source/DataprepEditor/Private/DataprepGraph/DataprepGraph.h
+++ b/Engine/Plugins/Enterprise/DataprepEditor/Source/DataprepEditor/Private/DataprepGraph/DataprepGraph.h
@@ -32,22 +32,11 @@ public:
 	/** Associate Dataprep asset with graph and add root node */
 	void Initialize(UDataprepAsset* InDataprepAsset);
 
-	void SetEditor(TSharedPtr<SGraphEditor>& InGraphEditor)
-	{
-		GraphEditor = InGraphEditor;
-	}
-
 	const UDataprepAsset* GetDataprepAsset() const { return DataprepAssetPtr.Get(); }
 	UDataprepAsset* GetDataprepAsset() { return DataprepAssetPtr.Get(); }
 
 
 private:
-	/** Dataprep graph editor associated with this graph */
-	TSharedPtr<SGraphEditor> GraphEditor;
-
-	/** Root node to enable the graph editing */
-	TStrongObjectPtr<UDataprepGraphRecipeNode> RecipeNode;
-
 	/** Dataprep asset associated with this graph */
 	TWeakObjectPtr<UDataprepAsset> DataprepAssetPtr;
 };
@@ -63,7 +52,7 @@ class UDataprepGraphRecipeNode : public UEdGraphNode
 
 public:
 	/** Get/Set associated widget */
-	TSharedPtr<SGraphNode> GetWidget() { return Widget; }
+	TSharedPtr<SGraphNode> GetWidget() { return Widget.Pin(); }
 	void SetWidget(TSharedPtr<SGraphNode> InWidget) { Widget = InWidget; }
 
 	// UEdGraphNode implementation
@@ -75,5 +64,5 @@ private:
 	/**
 	 * Associated widget displayed in the Dataprep graph editor
 	 */
-	TSharedPtr<SGraphNode> Widget;
+	TWeakPtr<SGraphNode> Widget;
 };
\ No newline at end of file
