Integrating CL#11177897 from //UE4/Dev-Enterprise

Fixed visual feedback when dragging step node over action
Added highlight when action node is selected
Added copy/paste (still WIP)

diff --git a/Engine/Plugins/Enterprise/DataprepEditor/Source/DataprepCore/Private/DataprepAsset.cpp b/Engine/Plugins/Enterprise/DataprepEditor/Source/DataprepCore/Private/DataprepAsset.cpp
--- a/Engine/Plugins/Enterprise/DataprepEditor/Source/DataprepCore/Private/DataprepAsset.cpp
+++ b/Engine/Plugins/Enterprise/DataprepEditor/Source/DataprepCore/Private/DataprepAsset.cpp
@@ -92,7 +92,7 @@ void UDataprepAsset::PostLoad()
 			}
 
 			// This Dataprep asset was never opened in the editor
-			if( StartNode == nullptr )
+			if( StartNode == nullptr && ActionAssets.Num() == 0)
 			{
 				IBlueprintNodeBinder::FBindingSet Bindings;
 				StartNode = UBlueprintNodeSpawner::Create<UK2Node_DataprepProducer>()->Invoke( PipelineGraph, Bindings, FVector2D(-100,0) );
@@ -476,7 +476,8 @@ void UDataprepAsset::OnDataprepBlueprintChanged( UBlueprint* InBlueprint )
 
 void UDataprepAsset::UpdateActions()
 {
-	ActionAssets.Empty(ActionAssets.Num());
+	TSet<UDataprepActionAsset*> ActionSet;
+	ActionSet.Append(ActionAssets);
 
 	UEdGraphPin* NodeOutPin= StartNode->FindPin(UEdGraphSchema_K2::PN_Then, EGPD_Output);
 	if( NodeOutPin && NodeOutPin->LinkedTo.Num() > 0 )
@@ -502,7 +503,8 @@ void UDataprepAsset::UpdateActions()
 			{
 				if( UDataprepActionAsset* DataprepAction = ActionNode->GetDataprepAction() )
 				{
-					ActionAssets.Add( DataprepAction );
+					DataprepAction->Rename(nullptr, this, REN_DoNotDirty | REN_DontCreateRedirectors | REN_NonTransactional);
+					ActionSet.Add( DataprepAction );
 				}
 			}
 
@@ -526,6 +528,8 @@ void UDataprepAsset::UpdateActions()
 		}
 	}
 
+	ActionAssets = ActionSet.Array();
+
 	CachedActionCount = ActionAssets.Num();
 }
 #endif
@@ -554,6 +558,42 @@ int32 UDataprepAsset::AddAction(const UDataprepActionAsset* InAction)
 	return INDEX_NONE;
 }
 
+int32 UDataprepAsset::AddActions(const TArray<const UDataprepActionAsset*>& InActions)
+{
+	if ( InActions.Num() > 0 && InActions[0] != nullptr )
+	{
+		Modify();
+
+		int32 PreviousActionCount = ActionAssets.Num();
+
+		for(const UDataprepActionAsset* InAction : InActions)
+		{
+			if(InAction)
+			{
+				UDataprepActionAsset* Action = DuplicateObject<UDataprepActionAsset>( InAction, this);
+				Action->SetLabel( InAction->GetLabel() );
+
+				ActionAssets.Add( Action );
+			}
+		}
+
+		CachedActionCount = ActionAssets.Num();
+
+		if(PreviousActionCount != CachedActionCount)
+		{
+			OnActionChanged.Broadcast(ActionAssets.Last(), FDataprepAssetChangeType::ActionAdded);
+
+			return ActionAssets.Num() - 1;
+		}
+	}
+
+	UE_LOG( LogDataprepCore, Error, TEXT("UDataprepAsset::AddActions: None of the action steps is invalid") );
+	ensure(false);
+
+	// Invalid
+	return INDEX_NONE;
+}
+
 int32 UDataprepAsset::AddActions(const TArray<const UDataprepActionStep*>& InActionSteps, bool bCreateOne)
 {
 	if ( InActionSteps.Num() > 0 && InActionSteps[0] != nullptr )
@@ -598,7 +638,7 @@ int32 UDataprepAsset::AddActions(const TArray<const UDataprepActionStep*>& InAct
 		}
 	}
 
-	UE_LOG( LogDataprepCore, Error, TEXT("UDataprepAsset::AddAction: None of the action steps is invalid") );
+	UE_LOG( LogDataprepCore, Error, TEXT("UDataprepAsset::AddActionSteps: None of the action steps is invalid") );
 	ensure(false);
 
 	// Invalid
@@ -629,7 +669,53 @@ bool UDataprepAsset::InsertAction(const UDataprepActionAsset* InAction, int32 In
 		return true;
 	}
 
-	UE_LOG( LogDataprepCore, Error, TEXT("UDataprepAsset::AddAction: The action is invalid") );
+	UE_LOG( LogDataprepCore, Error, TEXT("UDataprepAsset::InsertAction: The action is invalid") );
+	ensure(false);
+
+	// Invalid
+	return false;
+}
+
+bool UDataprepAsset::InsertActions(const TArray<const UDataprepActionAsset*>& InActions, int32 Index)
+{
+	if(!ActionAssets.IsValidIndex(Index))
+	{
+		UE_LOG( LogDataprepCore, Error, TEXT("UDataprepAsset::InsertActions: The index is invalid") );
+		return false;
+	}
+
+	if ( InActions.Num() > 0 && InActions[0] != nullptr )
+	{
+		Modify();
+
+		int32 PreviousActionCount = ActionAssets.Num();
+
+		int32 InsertIndex = Index;
+
+		for(const UDataprepActionAsset* InAction : InActions)
+		{
+			if(InAction)
+			{
+				UDataprepActionAsset* Action = DuplicateObject<UDataprepActionAsset>( InAction, this);
+				Action->SetLabel( InAction->GetLabel() );
+
+				ActionAssets.Insert( Action, InsertIndex );
+
+				++InsertIndex;
+			}
+		}
+
+		CachedActionCount = ActionAssets.Num();
+
+		if(PreviousActionCount != CachedActionCount)
+		{
+			OnActionChanged.Broadcast(ActionAssets.Last(), FDataprepAssetChangeType::ActionAdded);
+
+			return true;
+		}
+	}
+
+	UE_LOG( LogDataprepCore, Error, TEXT("UDataprepAsset::InsertActions: None of the actions is invalid") );
 	ensure(false);
 
 	// Invalid
diff --git a/Engine/Plugins/Enterprise/DataprepEditor/Source/DataprepCore/Public/DataprepAsset.h b/Engine/Plugins/Enterprise/DataprepEditor/Source/DataprepCore/Public/DataprepAsset.h
--- a/Engine/Plugins/Enterprise/DataprepEditor/Source/DataprepCore/Public/DataprepAsset.h
+++ b/Engine/Plugins/Enterprise/DataprepEditor/Source/DataprepCore/Public/DataprepAsset.h
@@ -70,6 +70,14 @@ public:
 	 */
 	int32 AddAction(const UDataprepActionAsset* Action);
 
+	/**
+	 * Add the actions to the Dataprep asset
+	 * @param Actions The array of action to add
+	 * @param bCreateOne Indicates if one or more action assets should be created. By default one is created
+	 * @return The index of the last added action or index none if the action is invalid
+	 */
+	int32 AddActions(const TArray<const UDataprepActionAsset*>& Actions);
+
 	/**
 	 * Creates an action from the array of action steps or one action per action steps
 	 * then add the action(s) to the Dataprep asset
@@ -87,6 +95,14 @@ public:
 	 */
 	bool InsertAction(const UDataprepActionAsset* InAction, int32 Index);
 
+	/**
+	 * Insert a copy of each action into the Dataprep asset at the requested index
+	 * @param Actions The array of actions to insert
+	 * @param Index The index at which the insertion must happen
+	 * @return True if the insertion is successful, false if the actions or the index are invalid
+	 */
+	bool InsertActions(const TArray<const UDataprepActionAsset*>& InActions, int32 Index);
+
 	/**
 	 * Creates an action from the array of action steps or one action per action steps
 	 * then insert the action(s) to the Dataprep asset at the requested index
diff --git a/Engine/Plugins/Enterprise/DataprepEditor/Source/DataprepEditor/Private/DataprepGraph/DataprepGraph.cpp b/Engine/Plugins/Enterprise/DataprepEditor/Source/DataprepEditor/Private/DataprepGraph/DataprepGraph.cpp
--- a/Engine/Plugins/Enterprise/DataprepEditor/Source/DataprepEditor/Private/DataprepGraph/DataprepGraph.cpp
+++ b/Engine/Plugins/Enterprise/DataprepEditor/Source/DataprepEditor/Private/DataprepGraph/DataprepGraph.cpp
@@ -32,50 +32,6 @@ void FDataprepEditor::CreateGraphEditor()
 {
 	if(UDataprepAsset* DataprepAsset = Cast<UDataprepAsset>(DataprepAssetInterfacePtr.Get()) )
 	{
-		if (!GraphEditorCommands.IsValid())
-		{
-			GraphEditorCommands = MakeShareable(new FUICommandList);
-
-			GraphEditorCommands->MapAction(FGenericCommands::Get().Rename,
-				FExecuteAction::CreateSP(this, &FDataprepEditor::OnRenameNode),
-				FCanExecuteAction::CreateSP(this, &FDataprepEditor::CanRenameNode)
-			);
-
-			GraphEditorCommands->MapAction(FGenericCommands::Get().SelectAll,
-				FExecuteAction::CreateSP(this, &FDataprepEditor::SelectAllNodes),
-				FCanExecuteAction::CreateSP(this, &FDataprepEditor::CanSelectAllNodes)
-			);
-
-			GraphEditorCommands->MapAction(FGenericCommands::Get().Delete,
-				FExecuteAction::CreateSP(this, &FDataprepEditor::DeleteSelectedPipelineNodes),
-				FCanExecuteAction::CreateSP(this, &FDataprepEditor::CanDeletePipelineNodes)
-			);
-
-			GraphEditorCommands->MapAction(FGenericCommands::Get().Copy,
-				FExecuteAction::CreateSP(this, &FDataprepEditor::CopySelectedNodes),
-				FCanExecuteAction::CreateSP(this, &FDataprepEditor::CanCopyNodes)
-			);
-
-			GraphEditorCommands->MapAction(FGenericCommands::Get().Cut,
-				FExecuteAction::CreateSP(this, &FDataprepEditor::CutSelectedNodes),
-				FCanExecuteAction::CreateSP(this, &FDataprepEditor::CanCutNodes)
-			);
-
-			GraphEditorCommands->MapAction(FGenericCommands::Get().Paste,
-				FExecuteAction::CreateSP(this, &FDataprepEditor::PasteNodes),
-				FCanExecuteAction::CreateSP(this, &FDataprepEditor::CanPasteNodes)
-			);
-
-			GraphEditorCommands->MapAction(FGenericCommands::Get().Duplicate,
-				FExecuteAction::CreateSP(this, &FDataprepEditor::DuplicateNodes),
-				FCanExecuteAction::CreateSP(this, &FDataprepEditor::CanDuplicateNodes)
-			);
-
-			GraphEditorCommands->MapAction(FGraphEditorCommands::Get().CreateComment,
-				FExecuteAction::CreateSP(this, &FDataprepEditor::OnCreateComment)
-			);
-		}
-
 		FGraphAppearanceInfo AppearanceInfo;
 		AppearanceInfo.CornerText = LOCTEXT("AppearanceCornerText", "DATAPREP");
 
@@ -111,7 +67,6 @@ void FDataprepEditor::CreateGraphEditor()
 		DataprepGraph->Initialize( DataprepAsset );
 
 		GraphEditor = SNew(SDataprepGraphEditor, DataprepAsset)
-			.AdditionalCommands(GraphEditorCommands)
 			.Appearance(AppearanceInfo)
 			.TitleBar(TitleBarWidget)
 			.GraphToEdit(DataprepGraph.Get())
diff --git a/Engine/Plugins/Enterprise/DataprepEditor/Source/DataprepEditor/Private/DataprepPipelineEditor.cpp b/Engine/Plugins/Enterprise/DataprepEditor/Source/DataprepEditor/Private/DataprepPipelineEditor.cpp
--- a/Engine/Plugins/Enterprise/DataprepEditor/Source/DataprepEditor/Private/DataprepPipelineEditor.cpp
+++ b/Engine/Plugins/Enterprise/DataprepEditor/Source/DataprepEditor/Private/DataprepPipelineEditor.cpp
@@ -313,11 +313,28 @@ void FDataprepEditor::CopySelectedNodes()
 		if (UEdGraphNode* Node = Cast<UEdGraphNode>(*SelectedIter))
 		{
 			Node->PrepareForCopying();
+
+			// Temporarily set DataprepActionAsset's owner as ActionNode to serialize it with the EdGraphNode
+			if(UK2Node_DataprepAction* ActionNode = Cast<UK2Node_DataprepAction>(Node))
+			{
+				UDataprepActionAsset* ActionAsset = ActionNode->GetDataprepAction();
+				ActionAsset->Rename(nullptr, ActionNode, REN_DoNotDirty | REN_DontCreateRedirectors | REN_NonTransactional);
+			}
 		}
 	}
 
 	FEdGraphUtilities::ExportNodesToText(SelectedNodes, /*out*/ ExportedText);
 	FPlatformApplicationMisc::ClipboardCopy(*ExportedText);
+
+	// Restore DataprepActionAssets' owner to the DataprepAsset
+	for (FGraphPanelSelectionSet::TConstIterator SelectedIter(SelectedNodes); SelectedIter; ++SelectedIter)
+	{
+		if(UK2Node_DataprepAction* ActionNode = Cast<UK2Node_DataprepAction>(*SelectedIter))
+		{
+			UDataprepActionAsset* ActionAsset = ActionNode->GetDataprepAction();
+			ActionAsset->Rename(nullptr, DataprepAssetInterfacePtr.Get(), REN_DoNotDirty | REN_DontCreateRedirectors | REN_NonTransactional);
+		}
+	}
 }
 
 bool FDataprepEditor::CanCopyNodes() const
diff --git a/Engine/Plugins/Enterprise/DataprepEditor/Source/DataprepEditor/Private/Widgets/DataprepGraph/SDataprepGraphActionNode.cpp b/Engine/Plugins/Enterprise/DataprepEditor/Source/DataprepEditor/Private/Widgets/DataprepGraph/SDataprepGraphActionNode.cpp
--- a/Engine/Plugins/Enterprise/DataprepEditor/Source/DataprepEditor/Private/Widgets/DataprepGraph/SDataprepGraphActionNode.cpp
+++ b/Engine/Plugins/Enterprise/DataprepEditor/Source/DataprepEditor/Private/Widgets/DataprepGraph/SDataprepGraphActionNode.cpp
@@ -12,8 +12,9 @@
 #include "Widgets/DataprepGraph/SDataprepGraphTrackNode.h"
 
 // Engine Includes
-#include "SGraphPanel.h"
+#include "GraphEditorSettings.h"
 #include "NodeFactory.h"
+#include "SGraphPanel.h"
 #include "Widgets/Colors/SColorBlock.h"
 #include "Widgets/SBoxPanel.h"
 #include "Widgets/SOverlay.h"
@@ -21,6 +22,82 @@
 
 #define LOCTEXT_NAMESPACE "DataprepGraphEditor"
 
+class SDataprepGraphActionProxyNode : public SGraphNode
+{
+public:
+	SLATE_BEGIN_ARGS(SDataprepGraphActionStepNode) {}
+	SLATE_END_ARGS()
+
+	void Construct(const FArguments& InArgs, const TSharedRef<SDataprepGraphActionNode>& InParentNode)
+	{
+		ParentNodePtr = InParentNode;
+		GraphNode = InParentNode->GetNodeObj();
+
+		SetCursor(EMouseCursor::Default);
+		UpdateGraphNode();
+	}
+
+	// SWidget interface
+	// End of SWidget interface
+
+	// SGraphNode interface
+	virtual void UpdateGraphNode() override
+	{
+		this->ContentScale.Bind( this, &SGraphNode::GetContentScale );
+		this->GetOrAddSlot( ENodeZone::Center )
+		.HAlign(HAlign_Fill)
+		.VAlign(VAlign_Center)
+		[
+			SNew(SVerticalBox)
+
+			+SVerticalBox::Slot()
+			.AutoHeight()
+			[
+				SNew(SColorBlock)
+				.Color( FLinearColor::Transparent )
+				.Size( TAttribute<FVector2D>::Create(TAttribute<FVector2D>::FGetter::CreateSP( this, &SDataprepGraphActionProxyNode::GetSize ) ) )
+			]
+		];
+	}
+
+	const FSlateBrush* GetShadowBrush(bool bSelected) const
+	{
+		return  FEditorStyle::GetNoBrush();
+	}
+	// End of SGraphNode interface
+
+	FVector2D GetSize()
+	{
+		FVector2D Size(10.f);
+
+		if(SDataprepGraphActionNode* ParentNode = ParentNodePtr.Pin().Get())
+		{
+			Size = ParentNode->GetCachedGeometry().GetLocalSize();
+
+			if(Size == FVector2D::ZeroVector)
+			{
+				Size = ParentNode->GetDesiredSize();
+				if(Size == FVector2D::ZeroVector)
+				{
+					Size.Set(10.f, 10.f);
+				}
+			}
+		}
+
+		return Size;
+	}
+
+	void SetPosition(const FVector2D& Position)
+	{
+		GraphNode->NodePosX = Position.X;
+		GraphNode->NodePosY = Position.Y;
+	}
+
+private:
+	/** Pointer to the SDataprepGraphTrackNode displayed in the graph editor  */
+	TWeakPtr<SDataprepGraphActionNode> ParentNodePtr;
+};
+
 /**
  * The SDataprepEmptyActionStepNode is a helper class that handles drag and drop event at
  * the bottom of the SDataprepGraphActionNode widget
@@ -55,8 +132,6 @@ public:
 				]
 			]
 		];
-
-		bShowInsertionSlot = false;
 	}
 
 	// SWidget Interface
@@ -68,7 +143,7 @@ public:
 			ParentTrackNodePtr.Pin()->OnDragLeave(DragDropEvent);
 
 			DragActionStepNodeOp->SetHoveredNode(ParentPtr.Pin()->GetNodeObj());
-			bShowInsertionSlot = true;
+			ParentPtr.Pin()->SetHoveredIndex( ParentPtr.Pin()->GetDataprepAction()->GetStepsCount() );
 		}
 
 		SCompoundWidget::OnDragEnter(MyGeometry, DragDropEvent);
@@ -80,7 +155,7 @@ public:
 		if(DragActionStepNodeOp.IsValid() && ParentPtr.IsValid())
 		{
 			DragActionStepNodeOp->SetHoveredNode(ParentPtr.Pin()->GetNodeObj());
-			bShowInsertionSlot = true;
+			ParentPtr.Pin()->SetHoveredIndex( ParentPtr.Pin()->GetDataprepAction()->GetStepsCount() );
 
 			return FReply::Handled();
 		}
@@ -96,15 +171,16 @@ public:
 			DragActionStepNodeOp->SetHoveredNode(nullptr);
 		}
 
-		bShowInsertionSlot = false;
+		ParentPtr.Pin()->SetHoveredIndex( INDEX_NONE );
 
 		SCompoundWidget::OnDragLeave(DragDropEvent);
 	}
 
 	virtual FReply OnDrop(const FGeometry& MyGeometry, const FDragDropEvent& DragDropEvent) override
 	{
-		bShowInsertionSlot = false;
-	
+		// Reset dragged index as drag is completed
+		ParentPtr.Pin()->SetDraggedIndex( INDEX_NONE );
+
 		// Process OnDrop if done by FDataprepDragDropOp
 		TSharedPtr<FDataprepDragDropOp> DragActionStepNodeOp = DragDropEvent.GetOperationAs<FDataprepDragDropOp>();
 		if (DragActionStepNodeOp.IsValid() && ParentPtr.IsValid())
@@ -122,20 +198,13 @@ public:
 		ParentTrackNodePtr = InParentTrackNode;
 	}
 
-	void ShowInsertionSlot(bool bShow)
-	{
-		bShowInsertionSlot = bShow;
-	}
-
 private:
 	FSlateColor GetBorderBackgroundColor() const
 	{
-		static FLinearColor Hovered = FDataprepEditorStyle::GetColor("DataprepActionStep.DragAndDrop");
-		return bShowInsertionSlot ? Hovered : FLinearColor::Transparent;
+		return ParentPtr.Pin()->GetInsertColor( ParentPtr.Pin()->GetDataprepAction()->GetStepsCount() );
 	}
 
 private:
-	bool bShowInsertionSlot;
 	TWeakPtr<SDataprepGraphActionNode> ParentPtr;
 	TWeakPtr<SDataprepGraphTrackNode> ParentTrackNodePtr;
 };
@@ -146,22 +215,19 @@ void SDataprepGraphActionNode::Construct(const FArguments& InArgs, UDataprepGrap
 	check(DataprepActionPtr.IsValid());
 
 	ExecutionOrder = InActionNode->GetExecutionOrder();
+	DraggedIndex = INDEX_NONE;
+	InsertIndex = INDEX_NONE;
 
 	DataprepActionPtr->GetOnStepsOrderChanged().AddSP(this, &SDataprepGraphActionNode::OnStepsChanged);
 
 	GraphNode = InActionNode;
 
+	ProxyNodePtr = SNew(SDataprepGraphActionProxyNode, SharedThis(this));
+
 	SetCursor(EMouseCursor::ResizeLeftRight);
 	UpdateGraphNode();
 }
 
-void SDataprepGraphActionNode::MoveTo(const FVector2D& DesiredPosition, FNodeSet& NodeFilter)
-{
-	FVector2D NewPosition = ParentTrackNodePtr.IsValid() ? ParentTrackNodePtr.Pin()->ComputeActionNodePosition( DesiredPosition ) : DesiredPosition;
-
-	SGraphNode::MoveTo(NewPosition, NodeFilter);
-}
-
 void SDataprepGraphActionNode::SetParentTrackNode(TSharedPtr<SDataprepGraphTrackNode> InParentTrackNode)
 {
 	ParentTrackNodePtr = InParentTrackNode;
@@ -184,6 +250,11 @@ void SDataprepGraphActionNode::UpdateExecutionOrder()
 	ExecutionOrder = Cast<UDataprepGraphActionNode>(GraphNode)->GetExecutionOrder();
 }
 
+void SDataprepGraphActionNode::UpdateProxyNode(const FVector2D& Position)
+{
+	ProxyNodePtr->SetPosition(Position);
+}
+
 TSharedRef<SWidget> SDataprepGraphActionNode::CreateNodeContentArea()
 {
 	if(DataprepActionPtr.IsValid())
@@ -195,21 +266,7 @@ TSharedRef<SWidget> SDataprepGraphActionNode::CreateNodeContentArea()
 		.AutoHeight()
 		[
 			ActionStepListWidgetPtr.ToSharedRef()
-		]
-
-		//+SVerticalBox::Slot()
-		//.AutoHeight()
-		//[
-		//	SNew(SHorizontalBox)
-		//	+SHorizontalBox::Slot()
-		//	.AutoWidth()
-		//	[
-		//		SNew(SColorBlock)
-		//		.Color(FLinearColor::Transparent)
-		//		.Size( FVector2D(200.f, 10.f) )
-		//	]
-		//]
-		;
+		];
 	}
 
 	return 	SNew(STextBlock)
@@ -217,6 +274,35 @@ TSharedRef<SWidget> SDataprepGraphActionNode::CreateNodeContentArea()
 	.Text( FText::FromString( TEXT("This node doesn't have a dataprep action!") ) );
 }
 
+const FSlateBrush* SDataprepGraphActionNode::GetShadowBrush(bool bSelected) const
+{
+	return  FEditorStyle::GetNoBrush();
+}
+
+int32 SDataprepGraphActionNode::OnPaint(const FPaintArgs& Args, const FGeometry& AllottedGeometry, const FSlateRect& MyCullingRect, FSlateWindowElementList& OutDrawElements, int32 LayerId, const FWidgetStyle& InWidgetStyle, bool bParentEnabled) const
+{
+	// Since only a proxy is in the graph panel, draw selection outline if applicable 
+	if(SGraphPanel* GraphPanel = GetOwnerPanel().Get())
+	{
+		if (GraphPanel->SelectionManager.SelectedNodes.Contains(GraphNode))
+		{
+			const FSlateBrush* ShadowBrush = FEditorStyle::GetBrush(TEXT("Graph.Node.ShadowSelected"));
+			const FVector2D NodeShadowSize = GetDefault<UGraphEditorSettings>()->GetShadowDeltaSize();
+
+			FSlateDrawElement::MakeBox(
+				OutDrawElements,
+				LayerId,
+				GetPaintSpaceGeometry().ToInflatedPaintGeometry(NodeShadowSize),
+				ShadowBrush,
+				ESlateDrawEffect::None,
+				FLinearColor(1.0f, 1.0f, 1.0f, 1.f)
+			);
+		}
+	}
+
+	return SGraphNode::OnPaint(Args, AllottedGeometry, MyCullingRect, OutDrawElements, LayerId, InWidgetStyle, bParentEnabled);
+}
+
 FReply SDataprepGraphActionNode::OnMouseButtonDown(const FGeometry& MyGeometry, const FPointerEvent& MouseEvent)
 {
 	GetOwnerPanel()->SelectionManager.ClickedOnNode(GraphNode, MouseEvent);
@@ -244,6 +330,8 @@ void SDataprepGraphActionNode::SetOwner(const TSharedRef<SGraphPanel>& OwnerPane
 	{
 		SGraphNode::SetOwner(OwnerPanel);
 
+		OwnerPanel->AddGraphNode(SharedThis(ProxyNodePtr.Get()));
+
 		for(TSharedPtr<SDataprepGraphActionStepNode>& ActionStepGraphNode : ActionStepGraphNodes)
 		{
 			if (ActionStepGraphNode.IsValid())
@@ -297,6 +385,32 @@ void SDataprepGraphActionNode::OnDragEnter(const FGeometry& MyGeometry, const FD
 	SGraphNode::OnDragEnter(MyGeometry, DragDropEvent);
 }
 
+FSlateColor SDataprepGraphActionNode::GetInsertColor(int32 Index)
+{
+	static const FSlateColor BackgroundColor = FDataprepEditorStyle::GetColor("DataprepActionStep.BackgroundColor");
+	static const FSlateColor DragAndDrop = FDataprepEditorStyle::GetColor("DataprepActionStep.DragAndDrop");
+
+	return Index == InsertIndex ? DragAndDrop : BackgroundColor;
+}
+
+void SDataprepGraphActionNode::SetDraggedIndex(int32 Index)
+{
+	DraggedIndex = Index;
+	InsertIndex = INDEX_NONE;
+}
+
+void SDataprepGraphActionNode::SetHoveredIndex(int32 Index)
+{
+	if(DraggedIndex == INDEX_NONE || Index == DataprepActionPtr->GetStepsCount())
+	{
+		InsertIndex = Index;
+	}
+	else
+	{
+		InsertIndex = Index > DraggedIndex ? Index + 1 : (Index < DraggedIndex ? Index : INDEX_NONE);
+	}
+}
+
 void SDataprepGraphActionNode::PopulateActionStepListWidget()
 {
 	if(!ActionStepListWidgetPtr.IsValid())
@@ -333,9 +447,7 @@ void SDataprepGraphActionNode::PopulateActionStepListWidget()
 
 		ActionStepNode->Initialize(DataprepAction, Index);
 
-		TSharedPtr<SDataprepGraphActionStepNode> ActionStepGraphNode = StaticCastSharedPtr<SDataprepGraphActionStepNode>(FNodeFactory::CreateNodeWidget(ActionStepNode));
-
-		ActionStepGraphNode->SetParentNode(SharedThis(this));
+		TSharedPtr<SDataprepGraphActionStepNode> ActionStepGraphNode = SNew(SDataprepGraphActionStepNode, ActionStepNode, SharedThis(this));
 
 		if(TrackNodePtr.IsValid())
 		{
@@ -377,6 +489,7 @@ void SDataprepGraphActionNode::OnStepsChanged()
 	if(DataprepActionPtr.IsValid())
 	{
 		PopulateActionStepListWidget();
+		ParentTrackNodePtr.Pin()->RefreshLayout();
 	}
 }
 
diff --git a/Engine/Plugins/Enterprise/DataprepEditor/Source/DataprepEditor/Private/Widgets/DataprepGraph/SDataprepGraphActionNode.h b/Engine/Plugins/Enterprise/DataprepEditor/Source/DataprepEditor/Private/Widgets/DataprepGraph/SDataprepGraphActionNode.h
--- a/Engine/Plugins/Enterprise/DataprepEditor/Source/DataprepEditor/Private/Widgets/DataprepGraph/SDataprepGraphActionNode.h
+++ b/Engine/Plugins/Enterprise/DataprepEditor/Source/DataprepEditor/Private/Widgets/DataprepGraph/SDataprepGraphActionNode.h
@@ -7,6 +7,7 @@
 #include "Widgets/DeclarativeSyntaxSupport.h"
 
 class SDataprepGraphTrackNode;
+class SDataprepGraphActionProxyNode;
 class SDataprepGraphActionStepNode;
 class SVerticalBox;
 class UDataprepActionAsset;
@@ -30,12 +31,13 @@ public:
 	virtual FCursorReply OnCursorQuery(const FGeometry& MyGeometry, const FPointerEvent& CursorEvent) const override;
 	virtual void SetOwner(const TSharedRef<SGraphPanel>& OwnerPanel) override;
 	virtual void OnDragEnter(const FGeometry& MyGeometry, const FDragDropEvent& DragDropEvent) override;
+	int32 OnPaint( const FPaintArgs& Args, const FGeometry& AllottedGeometry, const FSlateRect& MyCullingRect, FSlateWindowElementList& OutDrawElements, int32 LayerId, const FWidgetStyle& InWidgetStyle, bool bParentEnabled ) const override;
 	// End of SWidget interface
 
 	// SGraphNode interface
-	virtual void MoveTo( const FVector2D& NewPosition, FNodeSet& NodeFilter ) override;
 	virtual TSharedRef<SWidget> CreateNodeContentArea() override;
 	virtual FReply OnDrop( const FGeometry& MyGeometry, const FDragDropEvent& DragDropEvent ) override;
+	virtual const FSlateBrush* GetShadowBrush(bool bSelected) const override;
 	// End of SGraphNode interface
 
 	void SetParentTrackNode(TSharedPtr<SDataprepGraphTrackNode> InParentTrackNode);
@@ -45,6 +47,18 @@ public:
 
 	const UDataprepActionAsset* GetDataprepAction() const { return DataprepActionPtr.Get(); }
 
+	/** Update the proxy node with relative position in track node */
+	void UpdateProxyNode(const FVector2D& Position);
+
+	/** Callback used by insert nodes to determine their background color */
+	FSlateColor GetInsertColor(int32 Index);
+
+	/** Set index of step node being dragged */
+	void SetDraggedIndex(int32 Index);
+
+	/** Set index of step node being hovered */
+	void SetHoveredIndex(int32 Index);
+
 private:
 	/** Callback to handle changes on array of steps in action */
 	void OnStepsChanged();
@@ -67,4 +81,13 @@ private:
 
 	/** Array of pointers to the SDataprepGraphActionStepNode representing the associated action's steps */
 	TArray<TSharedPtr<SDataprepGraphActionStepNode>> ActionStepGraphNodes;
+
+	/** Pointer to the proxy SGraphNode inserted in the graph panel */
+	TSharedPtr<SDataprepGraphActionProxyNode> ProxyNodePtr;
+
+	/** Index of step node being dragged */
+	int32 DraggedIndex;
+
+	/** Index of insert widget to be highlighted */
+	int32 InsertIndex;
 };
diff --git a/Engine/Plugins/Enterprise/DataprepEditor/Source/DataprepEditor/Private/Widgets/DataprepGraph/SDataprepGraphActionStepNode.cpp b/Engine/Plugins/Enterprise/DataprepEditor/Source/DataprepEditor/Private/Widgets/DataprepGraph/SDataprepGraphActionStepNode.cpp
--- a/Engine/Plugins/Enterprise/DataprepEditor/Source/DataprepEditor/Private/Widgets/DataprepGraph/SDataprepGraphActionStepNode.cpp
+++ b/Engine/Plugins/Enterprise/DataprepEditor/Source/DataprepEditor/Private/Widgets/DataprepGraph/SDataprepGraphActionStepNode.cpp
@@ -14,6 +14,7 @@
 #include "SelectionSystem/DataprepFilter.h"
 #include "Widgets/DataprepGraph/SDataprepActionSteps.h"
 #include "Widgets/DataprepGraph/SDataprepFilter.h"
+#include "Widgets/DataprepGraph/SDataprepGraphActionNode.h"
 #include "Widgets/DataprepGraph/SDataprepGraphTrackNode.h"
 #include "Widgets/DataprepGraph/SDataprepOperation.h"
 
@@ -33,11 +34,11 @@
 
 #define LOCTEXT_NAMESPACE "DataprepGraphEditor"
 
-void SDataprepGraphActionStepNode::Construct(const FArguments& InArgs, UDataprepGraphActionStepNode* InActionStepNode)
+void SDataprepGraphActionStepNode::Construct(const FArguments& InArgs, UDataprepGraphActionStepNode* InActionStepNode, const TSharedPtr<SDataprepGraphActionNode>& InParent)
 {
 	StepIndex = InActionStepNode->GetStepIndex();
-	bShowInsertionSlot = false;
 
+	ParentNodePtr = InParent;
 	GraphNode = InActionStepNode;
 
 	SetCursor(EMouseCursor::ResizeUpDown);
@@ -150,10 +151,7 @@ void SDataprepGraphActionStepNode::UpdateGraphNode()
 
 FSlateColor SDataprepGraphActionStepNode::GetDragAndDropColor() const
 {
-	static const FLinearColor BackgroundColor = FDataprepEditorStyle::GetColor("DataprepActionStep.BackgroundColor");
-	static const FLinearColor DragAndDrop = FDataprepEditorStyle::GetColor("DataprepActionStep.DragAndDrop");
-
-	return bShowInsertionSlot ? DragAndDrop : BackgroundColor;
+	return ParentNodePtr.Pin()->GetInsertColor(StepIndex);
 }
 
 FSlateColor SDataprepGraphActionStepNode::GetBorderBackgroundColor() const
@@ -180,11 +178,6 @@ FReply SDataprepGraphActionStepNode::OnMouseButtonDown(const FGeometry& MyGeomet
 
 FReply SDataprepGraphActionStepNode::OnMouseMove(const FGeometry& SenderGeometry, const FPointerEvent& MouseEvent)
 {
-	if (!MouseEvent.IsMouseButtonDown(EKeys::LeftMouseButton) && bShowInsertionSlot)
-	{
-		bShowInsertionSlot = false;
-	}
-
 	return FReply::Unhandled();
 }
 
@@ -194,6 +187,7 @@ FReply SDataprepGraphActionStepNode::OnDragDetected(const FGeometry & MyGeometry
 	{
 		if (UDataprepActionStep* ActionStep = ActionStepNode->GetDataprepActionStep())
 		{
+			ParentNodePtr.Pin()->SetDraggedIndex(StepIndex);
 			return FReply::Handled().BeginDragDrop(FDataprepDragDropOp::New( GetOwnerPanel().ToSharedRef(), SharedThis(this)));
 		}
 	}
@@ -211,7 +205,7 @@ void SDataprepGraphActionStepNode::OnDragEnter(const FGeometry& MyGeometry, cons
 
 		// Inform the Drag and Drop operation that we are hovering over this node.
 		DragNodeOp->SetHoveredNode(GraphNode);
-		bShowInsertionSlot = DragNodeOp->IsValidDrop();
+		ParentNodePtr.Pin()->SetHoveredIndex(StepIndex);
 
 		return;
 	}
@@ -227,7 +221,7 @@ FReply SDataprepGraphActionStepNode::OnDragOver(const FGeometry& MyGeometry, con
 	{
 		// Inform the Drag and Drop operation that we are hovering over this node.
 		DragNodeOp->SetHoveredNode(GraphNode);
-		bShowInsertionSlot = DragNodeOp->IsValidDrop();
+		ParentNodePtr.Pin()->SetHoveredIndex(StepIndex);
 
 		return FReply::Handled();
 	}
@@ -242,7 +236,7 @@ void SDataprepGraphActionStepNode::OnDragLeave(const FDragDropEvent& DragDropEve
 	{
 		// Inform the Drag and Drop operation that we are not this widget anymore
 		DragNodeOp->SetHoveredNode(nullptr);
-		bShowInsertionSlot = false;
+		ParentNodePtr.Pin()->SetHoveredIndex(INDEX_NONE);
 
 		return;
 	}
@@ -252,7 +246,7 @@ void SDataprepGraphActionStepNode::OnDragLeave(const FDragDropEvent& DragDropEve
 
 FReply SDataprepGraphActionStepNode::OnDrop(const FGeometry& MyGeometry, const FDragDropEvent& DragDropEvent)
 {
-	bShowInsertionSlot = false;
+	ParentNodePtr.Pin()->SetDraggedIndex(INDEX_NONE);
 
 	// Process OnDrop if done by FDataprepDragDropOp
 	TSharedPtr<FDataprepDragDropOp> DragActionStepNodeOp = DragDropEvent.GetOperationAs<FDataprepDragDropOp>();
diff --git a/Engine/Plugins/Enterprise/DataprepEditor/Source/DataprepEditor/Private/Widgets/DataprepGraph/SDataprepGraphActionStepNode.h b/Engine/Plugins/Enterprise/DataprepEditor/Source/DataprepEditor/Private/Widgets/DataprepGraph/SDataprepGraphActionStepNode.h
--- a/Engine/Plugins/Enterprise/DataprepEditor/Source/DataprepEditor/Private/Widgets/DataprepGraph/SDataprepGraphActionStepNode.h
+++ b/Engine/Plugins/Enterprise/DataprepEditor/Source/DataprepEditor/Private/Widgets/DataprepGraph/SDataprepGraphActionStepNode.h
@@ -25,7 +25,7 @@ public:
 	SLATE_BEGIN_ARGS(SDataprepGraphActionStepNode) {}
 	SLATE_END_ARGS()
 
-	void Construct(const FArguments& InArgs, UDataprepGraphActionStepNode* InActionStepNode);
+	void Construct(const FArguments& InArgs, UDataprepGraphActionStepNode* InActionStepNode, const TSharedPtr<SDataprepGraphActionNode>& InParent);
 
 	// SWidget interface
 	virtual FReply OnMouseButtonDown( const FGeometry& MyGeometry, const FPointerEvent& MouseEvent ) override;
@@ -48,15 +48,6 @@ public:
 	virtual TSharedPtr<SWidget> GetStepTitleWidget() const;
 
 	void SetParentTrackNode(TSharedPtr<SDataprepGraphTrackNode> InParentTrackNode);
-	void SetParentNode(TSharedPtr<SDataprepGraphActionNode> InParentNode)
-	{
-		ParentNodePtr = InParentNode;
-	}
-
-	void ShowInsertionSlot(bool bShow)
-	{
-		bShowInsertionSlot = bShow;
-	}
 
 private:
 	/**
@@ -75,9 +66,6 @@ private:
 	/** Index of the represented action step in the list of the parent action */
 	int32 StepIndex;
 
-	/** Indicates if the widget is hovered during a drag and drop operation */
-	bool bShowInsertionSlot;
-
 	/** Pointer to the SDataprepGraphTrackNode displayed in the graph editor  */
 	TWeakPtr<SDataprepGraphTrackNode> ParentTrackNodePtr;
 
diff --git a/Engine/Plugins/Enterprise/DataprepEditor/Source/DataprepEditor/Private/Widgets/DataprepGraph/SDataprepGraphEditor.cpp b/Engine/Plugins/Enterprise/DataprepEditor/Source/DataprepEditor/Private/Widgets/DataprepGraph/SDataprepGraphEditor.cpp
--- a/Engine/Plugins/Enterprise/DataprepEditor/Source/DataprepEditor/Private/Widgets/DataprepGraph/SDataprepGraphEditor.cpp
+++ b/Engine/Plugins/Enterprise/DataprepEditor/Source/DataprepEditor/Private/Widgets/DataprepGraph/SDataprepGraphEditor.cpp
@@ -6,6 +6,7 @@
 #include "DataprepEditorLogCategory.h"
 #include "DataprepEditorStyle.h"
 #include "DataprepGraph/DataprepGraph.h"
+#include "DataprepGraph/DataprepGraphSchema.h"
 #include "DataprepGraph/DataprepGraphActionNode.h"
 #include "SchemaActions/DataprepDragDropOp.h"
 #include "Widgets/DataprepGraph/SDataprepGraphActionNode.h"
@@ -13,6 +14,14 @@
 #include "Widgets/DataprepGraph/SDataprepGraphTrackNode.h"
 #include "Widgets/DataprepWidgets.h"
 
+#include "EdGraph/EdGraphSchema.h"
+#include "EdGraphNode_Comment.h"
+#include "EdGraphUtilities.h"
+#include "Framework/Commands/GenericCommands.h"
+#include "Framework/Commands/UICommandList.h"
+#include "HAL/PlatformApplicationMisc.h"
+#include "Kismet2/Kismet2NameValidators.h"
+#include "ScopedTransaction.h"
 #include "SGraphPanel.h"
 #include "Widgets/Colors/SColorBlock.h"
 #include "Widgets/Layout/SBox.h"
@@ -41,10 +50,6 @@ TSharedPtr<class SGraphNode> SDataprepGraphEditorNodeFactory::CreateNode(UEdGrap
 	{
 		return SNew(SDataprepGraphActionNode, ActionNode);
 	}
-	else if (UDataprepGraphActionStepNode* ActionStepNode = Cast<UDataprepGraphActionStepNode>(Node))
-	{
-		return SNew(SDataprepGraphActionStepNode, ActionStepNode);
-	}
 
 	return nullptr;
 }
@@ -72,8 +77,10 @@ void SDataprepGraphEditor::Construct(const FArguments& InArgs, UDataprepAsset* I
 	check(InDataprepAsset);
 	DataprepAssetPtr = InDataprepAsset;
 
+	BuildCommandList();
+
 	SGraphEditor::FArguments Arguments;
-	Arguments._AdditionalCommands = InArgs._AdditionalCommands;
+	Arguments._AdditionalCommands = GraphEditorCommands;
 	Arguments._TitleBar = InArgs._TitleBar;
 	Arguments._GraphToEdit = InArgs._GraphToEdit;
 	Arguments._GraphEvents = InArgs._GraphEvents;
@@ -337,4 +344,356 @@ FReply SDataprepGraphEditor::OnDrop(const FGeometry& MyGeometry, const FDragDrop
 	return SGraphEditor::OnDrop(MyGeometry, DragDropEvent);
 }
 
+void SDataprepGraphEditor::BuildCommandList()
+{
+	if (!GraphEditorCommands.IsValid())
+	{
+		GraphEditorCommands = MakeShareable(new FUICommandList);
+
+		GraphEditorCommands->MapAction(FGenericCommands::Get().Rename,
+			FExecuteAction::CreateSP(this, &SDataprepGraphEditor::OnRenameNode),
+			FCanExecuteAction::CreateSP(this, &SDataprepGraphEditor::CanRenameNode)
+		);
+
+		GraphEditorCommands->MapAction(FGenericCommands::Get().SelectAll,
+			FExecuteAction::CreateSP(this, &SDataprepGraphEditor::SelectAllNodes),
+			FCanExecuteAction::CreateSP(this, &SDataprepGraphEditor::CanSelectAllNodes)
+		);
+
+		GraphEditorCommands->MapAction(FGenericCommands::Get().Delete,
+			FExecuteAction::CreateSP(this, &SDataprepGraphEditor::DeleteSelectedNodes),
+			FCanExecuteAction::CreateSP(this, &SDataprepGraphEditor::CanDeleteNodes)
+		);
+
+		GraphEditorCommands->MapAction(FGenericCommands::Get().Copy,
+			FExecuteAction::CreateSP(this, &SDataprepGraphEditor::CopySelectedNodes),
+			FCanExecuteAction::CreateSP(this, &SDataprepGraphEditor::CanCopyNodes)
+		);
+
+		GraphEditorCommands->MapAction(FGenericCommands::Get().Cut,
+			FExecuteAction::CreateSP(this, &SDataprepGraphEditor::CutSelectedNodes),
+			FCanExecuteAction::CreateSP(this, &SDataprepGraphEditor::CanCutNodes)
+		);
+
+		GraphEditorCommands->MapAction(FGenericCommands::Get().Paste,
+			FExecuteAction::CreateSP(this, &SDataprepGraphEditor::PasteNodes),
+			FCanExecuteAction::CreateSP(this, &SDataprepGraphEditor::CanPasteNodes)
+		);
+
+		GraphEditorCommands->MapAction(FGenericCommands::Get().Duplicate,
+			FExecuteAction::CreateSP(this, &SDataprepGraphEditor::DuplicateNodes),
+			FCanExecuteAction::CreateSP(this, &SDataprepGraphEditor::CanDuplicateNodes)
+		);
+
+		GraphEditorCommands->MapAction(FGraphEditorCommands::Get().CreateComment,
+			FExecuteAction::CreateSP(this, &SDataprepGraphEditor::OnCreateComment)
+		);
+	}
+}
+
+void SDataprepGraphEditor::OnRenameNode()
+{
+	const FGraphPanelSelectionSet SelectedNodes = GetSelectedNodes();
+
+	for (FGraphPanelSelectionSet::TConstIterator NodeIt(SelectedNodes); NodeIt; ++NodeIt)
+	{
+		UEdGraphNode* SelectedNode = Cast<UEdGraphNode>(*NodeIt);
+		if (SelectedNode != NULL && SelectedNode->bCanRenameNode)
+		{
+			IsNodeTitleVisible(SelectedNode, true);
+			break;
+		}
+	}
+}
+
+bool SDataprepGraphEditor::CanRenameNode() const
+{
+	const FGraphPanelSelectionSet SelectedNodes = GetSelectedNodes();
+	if ( SelectedNodes.Num() == 1)
+	{
+		if ( UDataprepGraphActionNode* SelectedNode = Cast<UDataprepGraphActionNode>(*SelectedNodes.CreateConstIterator()) )
+		{
+			return SelectedNode->bCanRenameNode;
+		}
+	}
+
+	return false;
+}
+
+void SDataprepGraphEditor::DeleteSelectedNodes()
+{
+	const FGraphPanelSelectionSet SelectedNodes = GetSelectedNodes();
+
+	TArray<int32> ActionsToDelete;
+	for(UObject* NodeObject : SelectedNodes)
+	{
+		if (UDataprepGraphActionNode* ActionNode = Cast<UDataprepGraphActionNode>(NodeObject))
+		{
+			if (ActionNode->CanUserDeleteNode() && ActionNode->GetDataprepActionAsset())
+			{
+				ActionsToDelete.Add(ActionNode->GetExecutionOrder());
+			}
+		}
+	}
+
+	if(ActionsToDelete.Num() > 0)
+	{
+		FScopedTransaction Transaction(FGenericCommands::Get().Delete->GetDescription());
+
+		bool bTransactionSuccessful = DataprepAssetPtr->RemoveActions(ActionsToDelete);
+
+		if (!bTransactionSuccessful)
+		{
+			Transaction.Cancel();
+		}
+	}
+}
+
+bool SDataprepGraphEditor::CanDeleteNodes() const
+{
+	const FGraphPanelSelectionSet SelectedNodes = GetSelectedNodes();
+
+	for (UObject* NodeObject : SelectedNodes)
+	{
+		// If any nodes allow deleting, then do not disable the delete option
+		UEdGraphNode* Node = Cast<UDataprepGraphActionNode>(NodeObject);
+		if (Node->CanUserDeleteNode())
+		{
+			return true;
+			break;
+		}
+	}
+
+	return false;
+}
+
+void SDataprepGraphEditor::CopySelectedNodes()
+{
+	// Export the selected nodes and place the text on the clipboard
+	const FGraphPanelSelectionSet SelectedNodes = GetSelectedNodes();
+	TSet<UObject*> ActionsToCopy;
+
+	for (UObject* NodeObject : SelectedNodes)
+	{
+		if (UDataprepGraphActionNode* ActionNode = Cast<UDataprepGraphActionNode>(NodeObject))
+		{
+			UDataprepActionAsset* ActionAsset = ActionNode->GetDataprepActionAsset();
+			// Temporarily set DataprepActionAsset's owner as ActionNode to serialize it with the EdGraphNode
+			ActionAsset->Rename(nullptr, ActionNode, REN_DoNotDirty | REN_DontCreateRedirectors | REN_NonTransactional);
+
+			ActionNode->PrepareForCopying();
+			ActionsToCopy.Add(ActionNode);
+		}
+	}
+
+	if(ActionsToCopy.Num() > 0)
+	{
+		FString ExportedText;
+
+		FEdGraphUtilities::ExportNodesToText(ActionsToCopy, /*out*/ ExportedText);
+		FPlatformApplicationMisc::ClipboardCopy(*ExportedText);
+
+		// Restore DataprepActionAssets' owner to the DataprepAsset
+		for (UObject* NodeObject : ActionsToCopy)
+		{
+			if (UDataprepGraphActionNode* ActionNode = Cast<UDataprepGraphActionNode>(NodeObject))
+			{
+				UDataprepActionAsset* ActionAsset = ActionNode->GetDataprepActionAsset();
+				ActionAsset->Rename(nullptr, DataprepAssetPtr.Get(), REN_DoNotDirty | REN_DontCreateRedirectors | REN_NonTransactional);
+			}
+		}
+	}
+}
+
+bool SDataprepGraphEditor::CanCopyNodes() const
+{
+	// If any of the nodes can be duplicated then we should allow copying
+	const FGraphPanelSelectionSet SelectedNodes = GetSelectedNodes();
+	for (UObject* NodeObject : SelectedNodes)
+	{
+		UDataprepGraphActionNode* Node = Cast<UDataprepGraphActionNode>(NodeObject);
+		if ((Node != NULL) && Node->CanDuplicateNode())
+		{
+			return true;
+		}
+	}
+	return false;
+}
+
+void SDataprepGraphEditor::CutSelectedNodes()
+{
+	CopySelectedNodes();
+	// Cut should only delete nodes that can be duplicated
+	DeleteSelectedDuplicatableNodes();
+}
+
+bool SDataprepGraphEditor::CanCutNodes() const
+{
+	return CanCopyNodes() && CanDeleteNodes();
+}
+
+void SDataprepGraphEditor::PasteNodes()
+{
+	ClearSelectionSet();
+
+	// Grab the text to paste from the clipboard.
+	FString TextToImport;
+	FPlatformApplicationMisc::ClipboardPaste(TextToImport);
+
+	// Create temporary graph
+	FName UniqueGraphName = MakeUniqueObjectName( GetTransientPackage(), UWorld::StaticClass(), FName( *(LOCTEXT("DataprepGraph", "TempGraph").ToString()) ) );
+	TStrongObjectPtr<UDataprepGraph> DataprepGraph = TStrongObjectPtr<UDataprepGraph>( NewObject< UDataprepGraph >(GetTransientPackage(), UniqueGraphName) );
+	DataprepGraph->Schema = UDataprepGraphSchema::StaticClass();
+
+	// Import the nodes
+	// #ueent_wip: The FEdGraphUtilities::ImportNodesFromTex could be replaced
+	TSet<UEdGraphNode*> PastedNodes;
+	FEdGraphUtilities::ImportNodesFromText(DataprepGraph.Get(), TextToImport, /*out*/ PastedNodes);
+
+	TArray<const UDataprepActionAsset*> Actions;
+	for(UEdGraphNode* Node : PastedNodes)
+	{
+		UDataprepGraphActionNode* ActionNode = Cast<UDataprepGraphActionNode>(Node);
+		if (ActionNode && ActionNode->CanDuplicateNode() && ActionNode->GetDataprepActionAsset())
+		{
+			Actions.Add(ActionNode->GetDataprepActionAsset());
+		}
+	}
+
+	if(Actions.Num() > 0)
+	{
+		FScopedTransaction Transaction(FGenericCommands::Get().Paste->GetDescription());
+
+		if(DataprepAssetPtr->AddActions(Actions) == INDEX_NONE)
+		{
+			Transaction.Cancel();
+		}
+	}
+}
+
+bool SDataprepGraphEditor::CanPasteNodes() const
+{
+	FString ClipboardContent;
+	FPlatformApplicationMisc::ClipboardPaste(ClipboardContent);
+
+	return FEdGraphUtilities::CanImportNodesFromText(GetCurrentGraph(), ClipboardContent);
+}
+
+void SDataprepGraphEditor::DuplicateNodes()
+{
+	// Copy and paste current selection
+	CopySelectedNodes();
+	PasteNodes();
+}
+
+bool SDataprepGraphEditor::CanDuplicateNodes() const
+{
+	return CanCopyNodes();
+}
+
+void SDataprepGraphEditor::DeleteSelectedDuplicatableNodes()
+{
+	// Cache off the old selection
+	const FGraphPanelSelectionSet OldSelectedNodes = GetSelectedNodes();
+
+	// Clear the selection and only select the nodes that can be duplicated
+	FGraphPanelSelectionSet CurrentSelection;
+	ClearSelectionSet();
+
+	FGraphPanelSelectionSet RemainingNodes;
+	for (FGraphPanelSelectionSet::TConstIterator SelectedIter(OldSelectedNodes); SelectedIter; ++SelectedIter)
+	{
+		UEdGraphNode* Node = Cast<UEdGraphNode>(*SelectedIter);
+		if ((Node != NULL) && Node->CanDuplicateNode())
+		{
+			SetNodeSelection(Node, true);
+		}
+		else
+		{
+			RemainingNodes.Add(Node);
+		}
+	}
+
+	// Delete the nodes which can be duplicated
+	DeleteSelectedNodes();
+
+	// Reselect whatever is left from the original selection after the deletion
+	ClearSelectionSet();
+
+	for (FGraphPanelSelectionSet::TConstIterator SelectedIter(RemainingNodes); SelectedIter; ++SelectedIter)
+	{
+		if (UEdGraphNode* Node = Cast<UEdGraphNode>(*SelectedIter))
+		{
+			SetNodeSelection(Node, true);
+		}
+	}
+}
+
+void SDataprepGraphEditor::OnCreateComment()
+{
+	if (UEdGraph* Graph = GetCurrentGraph())
+	{
+		if (const UEdGraphSchema* Schema = Graph->GetSchema())
+		{
+			// Add menu item for creating comment boxes
+			UEdGraphNode_Comment* CommentTemplate = NewObject<UEdGraphNode_Comment>();
+
+			FVector2D SpawnLocation = GetPasteLocation();
+
+			FSlateRect Bounds;
+			if ( GetBoundsForSelectedNodes(Bounds, 50.0f) )
+			{
+				CommentTemplate->SetBounds(Bounds);
+				SpawnLocation.X = CommentTemplate->NodePosX;
+				SpawnLocation.Y = CommentTemplate->NodePosY;
+			}
+
+			UEdGraphNode* NewNode = FEdGraphSchemaAction_NewNode::SpawnNodeFromTemplate<UEdGraphNode_Comment>(Graph, CommentTemplate, SpawnLocation, /** bSelectNewNode */ true);
+		}
+	}
+}
+
+
+bool SDataprepGraphEditor::OnNodeVerifyTitleCommit(const FText& NewText, UEdGraphNode* NodeBeingChanged, FText& OutErrorMessage)
+{
+	bool bValid(false);
+
+	if (NodeBeingChanged && NodeBeingChanged->bCanRenameNode)
+	{
+		// Clear off any existing error message 
+		NodeBeingChanged->ErrorMsg.Empty();
+		NodeBeingChanged->bHasCompilerMessage = false;
+
+		TSharedPtr<INameValidatorInterface> NameEntryValidator = NodeBeingChanged->MakeNameValidator();
+
+		check( NameEntryValidator.IsValid() );
+
+		EValidatorResult VResult = NameEntryValidator->IsValid(NewText.ToString(), true);
+		if (VResult == EValidatorResult::Ok)
+		{
+			bValid = true;
+		}
+		//else if (PipelineView.IsValid())
+		//{
+		//	EValidatorResult Valid = NameEntryValidator->IsValid(NewText.ToString(), false);
+
+		//	NodeBeingChanged->bHasCompilerMessage = true;
+		//	NodeBeingChanged->ErrorMsg = NameEntryValidator->GetErrorString(NewText.ToString(), Valid);
+		//	NodeBeingChanged->ErrorType = EMessageSeverity::Error;
+		//}
+	}
+
+	return bValid;
+}
+
+void SDataprepGraphEditor::OnNodeTitleCommitted(const FText& NewText, ETextCommit::Type CommitInfo, UEdGraphNode* NodeBeingChanged)
+{
+	if (NodeBeingChanged)
+	{
+		const FScopedTransaction Transaction(NSLOCTEXT("RenameNode", "RenameNode", "Rename Node"));
+		NodeBeingChanged->Modify();
+		NodeBeingChanged->OnRenameNode(NewText.ToString());
+	}
+}
+
 #undef LOCTEXT_NAMESPACE
\ No newline at end of file
diff --git a/Engine/Plugins/Enterprise/DataprepEditor/Source/DataprepEditor/Private/Widgets/DataprepGraph/SDataprepGraphEditor.h b/Engine/Plugins/Enterprise/DataprepEditor/Source/DataprepEditor/Private/Widgets/DataprepGraph/SDataprepGraphEditor.h
--- a/Engine/Plugins/Enterprise/DataprepEditor/Source/DataprepEditor/Private/Widgets/DataprepGraph/SDataprepGraphEditor.h
+++ b/Engine/Plugins/Enterprise/DataprepEditor/Source/DataprepEditor/Private/Widgets/DataprepGraph/SDataprepGraphEditor.h
@@ -18,6 +18,7 @@
 #include "Widgets/Layout/SConstraintCanvas.h"
 #include "Widgets/SWidget.h"
 
+class FUICommandList;
 class SDataprepGraphActionNode;
 class SDataprepGraphTrackNode;
 class SDataprepGraphTrackWidget;
@@ -71,6 +72,34 @@ public:
 	static void RegisterFactories();
 	static void UnRegisterFactories();
 
+	/** Set of methods necessary for copy/paste of action nodes*/
+	void OnRenameNode();
+	bool CanRenameNode() const;
+
+	bool CanSelectAllNodes() const { return true; }
+
+	void DeleteSelectedNodes();
+	bool CanDeleteNodes() const;
+
+	void CopySelectedNodes();
+	bool CanCopyNodes() const;
+
+	void CutSelectedNodes();
+	bool CanCutNodes() const;
+
+	void PasteNodes();
+	bool CanPasteNodes() const;
+
+	void DuplicateNodes();
+	bool CanDuplicateNodes() const;
+
+	void OnCreateComment();
+
+	void DeleteSelectedDuplicatableNodes();
+
+	bool OnNodeVerifyTitleCommit(const FText& NewText, UEdGraphNode* NodeBeingChanged, FText& OutErrorMessage);
+	void OnNodeTitleCommitted(const FText& NewText, ETextCommit::Type CommitInfo, UEdGraphNode* NodeBeingChanged);
+
 private:
 	/** Recompute the layout of the displayed graph after a pan, resize and/or zoom */
 	void UpdateLayout( const FVector2D& LocalSize, const FVector2D& Location, float ZoomAmount );
@@ -78,6 +107,8 @@ private:
 	/** Recompute the boundaries of a the displayed graph */
 	void UpdateBoundaries(const FVector2D& LocalSize, float ZoomAmount);
 
+	void BuildCommandList();
+
 private:
 	/** When false, indicates the graph editor has not been drawn yet */
 	mutable bool bIsComplete;
@@ -106,6 +137,9 @@ private:
 	/** Size of graph being displayed */
 	mutable TWeakPtr<SDataprepGraphTrackNode> TrackGraphNodePtr;
 
+	/** Command list associated with this graph editor */
+	TSharedPtr<FUICommandList> GraphEditorCommands;
+
 	bool bCachedControlKeyDown;
 
 	/** Padding used on the borders of the canvas */
@@ -117,5 +151,6 @@ private:
 	/** Factory to create the associated SGraphNode classes for Dataprep graph's UEdGraph classes */
 	static TSharedPtr<SDataprepGraphEditorNodeFactory> NodeFactory;
 
+
 	friend SDataprepGraphTrackNode;
 };
diff --git a/Engine/Plugins/Enterprise/DataprepEditor/Source/DataprepEditor/Private/Widgets/DataprepGraph/SDataprepGraphTrackNode.cpp b/Engine/Plugins/Enterprise/DataprepEditor/Source/DataprepEditor/Private/Widgets/DataprepGraph/SDataprepGraphTrackNode.cpp
--- a/Engine/Plugins/Enterprise/DataprepEditor/Source/DataprepEditor/Private/Widgets/DataprepGraph/SDataprepGraphTrackNode.cpp
+++ b/Engine/Plugins/Enterprise/DataprepEditor/Source/DataprepEditor/Private/Widgets/DataprepGraph/SDataprepGraphTrackNode.cpp
@@ -136,6 +136,7 @@ public:
 	int32 DragIndex;
 	FVector2D LastDragPosition;
 	FVector2D AbscissaRange;
+	FVector2D TrackOffset;
 	bool bIsCopying;
 
 	static float SlotSpacing;
@@ -313,10 +314,12 @@ bool SDataprepGraphTrackNode::RefreshLayout()
 	{
 		ensure(TrackWidgetPtr.IsValid());
 
-		FVector2D TrackWidgetOrigin = TrackWidgetPtr->GetCachedGeometry().LocalToAbsolute(FVector2D::ZeroVector);
-		FVector2D TrackNodeOrigin = GetCachedGeometry().LocalToAbsolute(FVector2D::ZeroVector);
+		const FVector2D TrackWidgetOrigin = TrackWidgetPtr->GetCachedGeometry().LocalToAbsolute(FVector2D::ZeroVector);
+		const FVector2D TrackNodeOrigin = GetCachedGeometry().LocalToAbsolute(FVector2D::ZeroVector);
 
-		TrackWidgetOffset = TrackWidgetOrigin - TrackNodeOrigin;
+		// Weird: Ultimate offset between track widget and node is 5 but it may take more draw call to stabilize
+		//		  Forcing 5.
+		TrackWidgetPtr->TrackOffset = FVector2D(5.f)/*TrackWidgetOrigin - TrackNodeOrigin*/ + GetPosition();
 
 		TrackWidgetPtr->RefreshLayout(OwnerGraphPanelPtr.Pin().Get());
 
@@ -411,14 +414,6 @@ FVector2D SDataprepGraphTrackNode::Update(const FVector2D& LocalSize, float Zoom
 	return FVector2D(SDataprepGraphTrackWidget::NodeDesiredWidth);
 }
 
-FVector2D SDataprepGraphTrackNode::ComputeActionNodePosition(const FVector2D& InPosition)
-{
-	ensure(TrackWidgetPtr);
-
-	const float NewAbscissa = TrackWidgetPtr->ValidateNodeAbscissa(InPosition.X - NodePadding.Left - TrackWidgetOffset.X);
-	return FVector2D( NewAbscissa + NodePadding.Left + TrackWidgetOffset.X, GraphNode->NodePosY + TrackWidgetOffset.Y);
-}
-
 void SDataprepGraphTrackNode::OnStartNodeDrag(const TSharedRef<SDataprepGraphActionNode>& ActionNode)
 {
 	bNodeDragging = true;
@@ -560,7 +555,7 @@ void SDataprepGraphTrackWidget::Construct(const FArguments& InArgs, TSharedPtr<S
 
 	TrackCanvas->AddSlot()
 	.Anchors( FAnchors( 0.f, 0.f, 1.f, 0.f ) )
-	.Offset( FMargin(0.f, 5.f, 0.f, 0.f) )
+	.Offset( FMargin(0.f, 10.f, 0.f, 0.f) )
 	.AutoSize(true)
 	.Alignment(FVector2D::ZeroVector)
 	.ZOrder(1)
@@ -609,7 +604,7 @@ void SDataprepGraphTrackWidget::Construct(const FArguments& InArgs, TSharedPtr<S
 			{
 				SConstraintCanvas::FSlot& Slot = TrackCanvas->AddSlot()
 				.Anchors( FAnchors( 0., 0.f, 0.f, 0.f ) )
-				.Offset( FMargin(LeftOffset, 5.f, SlotSpacing, TrackDesiredHeight ) )
+				.Offset( FMargin(LeftOffset, 10.f, SlotSpacing, TrackDesiredHeight ) )
 				.Alignment(FVector2D::ZeroVector)
 				.ZOrder(2)
 				[
@@ -793,6 +788,9 @@ void SDataprepGraphTrackWidget::RefreshLayout(SGraphPanel* GraphPanel)
 			ActionNode
 		];
 
+		// Update action's proxy node registered to graph panel
+		StaticCastSharedRef<SDataprepGraphActionNode>(ActionNode)->UpdateProxyNode(FVector2D(DropSlotOffset.Left, 0.f) + TrackOffset);
+
 		const FVector2D Size = GetNodeSize(ActionNode);
 
 		ActionSlots[Index]->OffsetAttr.Set( FMargin(DropSlotOffset.Left, 0.f, Size.X, Size.Y));
diff --git a/Engine/Plugins/Enterprise/DataprepEditor/Source/DataprepEditor/Private/Widgets/DataprepGraph/SDataprepGraphTrackNode.h b/Engine/Plugins/Enterprise/DataprepEditor/Source/DataprepEditor/Private/Widgets/DataprepGraph/SDataprepGraphTrackNode.h
--- a/Engine/Plugins/Enterprise/DataprepEditor/Source/DataprepEditor/Private/Widgets/DataprepGraph/SDataprepGraphTrackNode.h
+++ b/Engine/Plugins/Enterprise/DataprepEditor/Source/DataprepEditor/Private/Widgets/DataprepGraph/SDataprepGraphTrackNode.h
@@ -89,12 +89,6 @@ public:
 	/** Recompute the boundaries of the graph based on the new size and the new zoom factor */
 	FVector2D Update(const FVector2D& LocalSize, float ZoomAmount);
 
-	/**
-	 * Computes a new position for an action node based on the dimension of the graph
-	 * Mainly keeps the node at the right height
-	 */
-	FVector2D ComputeActionNodePosition(const FVector2D& InPosition);
-
 	void OnControlKeyChanged(bool bControlKeyDown);
 
 	/** Initiates the horizontal drag of an action node */
@@ -134,9 +128,6 @@ private:
 	/** Weak pointer to the Dataprep asset holding the displayed actions */
 	TWeakObjectPtr<UDataprepAsset> DataprepAssetPtr;
 
-	/** Size of the section of the track for left padding */
-	FVector2D TrackWidgetOffset;
-
 	/** Range for abscissa of action nodes */
 	FVector2D AbscissaRange;
 
