Integrating CL #11068492 //UE4/Dev-Enterprise to Release-4.25 (//UE4/Release-4.25)

Fixed the dataprep parameterisation. The type checking cache for a binding is not serialiized currently this will need to be fixed later.

diff --git a/Engine/Plugins/Enterprise/DataprepEditor/Source/DataprepCore/Private/Parameterization/DataprepParameterization.cpp b/Engine/Plugins/Enterprise/DataprepEditor/Source/DataprepCore/Private/Parameterization/DataprepParameterization.cpp
--- a/Engine/Plugins/Enterprise/DataprepEditor/Source/DataprepCore/Private/Parameterization/DataprepParameterization.cpp
+++ b/Engine/Plugins/Enterprise/DataprepEditor/Source/DataprepCore/Private/Parameterization/DataprepParameterization.cpp
@@ -163,9 +163,9 @@ namespace DataprepParameterization
 			}
 
 			if ( NumberOfObject == ValueTypeValidationData.Num() )
-				{
-					ValueTypeValidationData.Add( CurrentProperty->GetClass() );
-				}
+			{
+				ValueTypeValidationData.Add( CurrentProperty->GetClass() );
+			}
 		}
 	}
 
@@ -522,40 +522,6 @@ FDataprepParameterizationBinding::FDataprepParameterizationBinding(UDataprepPara
 	}
 }
 
-bool FDataprepParameterizationBinding::Serialize(FArchive& Ar)
-{
-	check(false); // @todo FProperties: if we never hit this, we don't need this
-
-	Ar.UsingCustomVersion(FCoreObjectVersion::GUID);
-
-	UScriptStruct* Struct = FDataprepParameterizationBinding::StaticStruct();
-	if (Ar.IsLoading() || Ar.IsSaving())
-	{
-		Struct->SerializeTaggedProperties(Ar, (uint8*)this, Struct, nullptr);
-	}
-
-#if WITH_EDITORONLY_DATA
-	//Take old data and put it in new data structure
-	if (Ar.IsLoading() && Ar.CustomVer(FCoreObjectVersion::GUID) < FCoreObjectVersion::FProperties)
-	{
-		//if (ValueType_DEPRECATED)
-		//{
-		//	ValueType = FFieldClass::GetNameToFieldClassMap().FindRef(ValueType_DEPRECATED->GetFName());
-		//}
-		//else
-		//{
-		//	ValueType = nullptr;
-		//}
-	}
-	else
-#endif // WITH_EDITORONLY_DATA
-	{
-		ValueTypeValidationData.Serialize(Ar);
-	}
-
-	return true;
-}
-
 bool FDataprepParameterizationBinding::operator==(const FDataprepParameterizationBinding& Other) const
 {
 	// The value type validation data shouldn't matter when comparing binding
diff --git a/Engine/Plugins/Enterprise/DataprepEditor/Source/DataprepCore/Private/Parameterization/DataprepParameterization.h b/Engine/Plugins/Enterprise/DataprepEditor/Source/DataprepCore/Private/Parameterization/DataprepParameterization.h
--- a/Engine/Plugins/Enterprise/DataprepEditor/Source/DataprepCore/Private/Parameterization/DataprepParameterization.h
+++ b/Engine/Plugins/Enterprise/DataprepEditor/Source/DataprepCore/Private/Parameterization/DataprepParameterization.h
@@ -60,14 +60,6 @@ struct FValueTypeValidationData
 	}
 };
 
-//template<> struct TStructOpsTypeTraits<FValueTypeValidationData> : public TStructOpsTypeTraitsBase2<FValueTypeValidationData>
-//{
-//	enum
-//	{
-//		WithSerializer = true
-//	};
-//};
-
 /**
  * The parameterization binding is a struct that hold an object and the property path to the parameterized property
  * It also hold a array to validate the that value type of the parameterized property didn't change since it's creation
@@ -101,16 +93,6 @@ struct FDataprepParameterizationBinding
 	// Value Type Validation Array. This is the result of a depth first search on the parametrized property
 	//UPROPERTY() // @todo FProp: do we need this to be an FProp? (is this struct being serialized)
 	FValueTypeValidationData ValueTypeValidationData;
-
-	bool Serialize(FArchive& Ar);
-};
-
-template<> struct TStructOpsTypeTraits<FDataprepParameterizationBinding> : public TStructOpsTypeTraitsBase2<FDataprepParameterizationBinding>
-{
-	enum
-	{
-		WithSerializer = true
-	};
 };
 
 uint32 GetTypeHash(const FDataprepParameterizationBinding& Binding);
