Integrating CL#11188985 from  //UE4/Dev-Enterprise

Dataprep filter preview first pass. Cleaned up some dead code and changed the way the dataprep parameterizable object notify the parameterization.

diff --git a/Engine/Plugins/Enterprise/DataprepEditor/Source/DataprepCore/Private/DataprepActionAsset.cpp b/Engine/Plugins/Enterprise/DataprepEditor/Source/DataprepCore/Private/DataprepActionAsset.cpp
--- a/Engine/Plugins/Enterprise/DataprepEditor/Source/DataprepCore/Private/DataprepActionAsset.cpp
+++ b/Engine/Plugins/Enterprise/DataprepEditor/Source/DataprepCore/Private/DataprepActionAsset.cpp
@@ -145,7 +145,7 @@ int32 UDataprepActionAsset::AddStep(TSubclassOf<UDataprepParameterizableObject>
 			ActionStep->PathOfStepObjectClass = ActionStep->StepObject->GetClass();
 			ActionStep->bIsEnabled = true;
 			Steps.Add(ActionStep);
-			OnStepsChanged.Broadcast();
+			OnStepsOrderChanged.Broadcast();
 			return Steps.Num() - 1;
 		}
 		else if ( ValidRootClass == UDataprepFetcher::StaticClass() )
@@ -160,7 +160,7 @@ int32 UDataprepActionAsset::AddStep(TSubclassOf<UDataprepParameterizableObject>
 				ActionStep->PathOfStepObjectClass = ActionStep->StepObject->GetClass();
 				ActionStep->bIsEnabled = true;
 				Steps.Add( ActionStep );
-				OnStepsChanged.Broadcast();
+				OnStepsOrderChanged.Broadcast();
 				return Steps.Num() - 1;
 			}
 			else
@@ -185,7 +185,7 @@ int32 UDataprepActionAsset::AddStep(const UDataprepActionStep* InActionStep)
 		Modify();
 		UDataprepActionStep* ActionStep = DuplicateObject<UDataprepActionStep>( InActionStep, this );
 		Steps.Add( ActionStep );
-		OnStepsChanged.Broadcast();
+		OnStepsOrderChanged.Broadcast();
 		return Steps.Num() - 1;
 	}
 
@@ -212,7 +212,7 @@ int32 UDataprepActionAsset::AddSteps(const TArray<const UDataprepActionStep*>& I
 			}
 		}
 
-		OnStepsChanged.Broadcast();
+		OnStepsOrderChanged.Broadcast();
 		return Steps.Num() - 1;
 	}
 
@@ -236,7 +236,7 @@ int32 UDataprepActionAsset::AddStep(const UDataprepParameterizableObject* StepOb
 		ActionStep->PathOfStepObjectClass = ActionStep->StepObject->GetClass();
 		ActionStep->bIsEnabled = true;
 		Steps.Add(ActionStep);
-		OnStepsChanged.Broadcast();
+		OnStepsOrderChanged.Broadcast();
 		return Steps.Num() - 1;
 	}
 	else
@@ -262,7 +262,7 @@ bool UDataprepActionAsset::InsertStep(const UDataprepActionStep* InActionStep, i
 		Modify();
 		UDataprepActionStep* ActionStep = DuplicateObject<UDataprepActionStep>( InActionStep, this);
 		Steps.Insert( ActionStep, Index );
-		OnStepsChanged.Broadcast();
+		OnStepsOrderChanged.Broadcast();
 		return true;
 	}
 
@@ -296,7 +296,7 @@ bool UDataprepActionAsset::InsertSteps(const TArray<const UDataprepActionStep*>&
 			}
 		}
 
-		OnStepsChanged.Broadcast();
+		OnStepsOrderChanged.Broadcast();
 		return true;
 	}
 
@@ -363,7 +363,7 @@ bool UDataprepActionAsset::MoveStep(int32 StepIndex, int32 DestinationIndex)
 
 	if ( DataprepCorePrivateUtils::MoveArrayElement( Steps, StepIndex, DestinationIndex ) )
 	{
-		OnStepsChanged.Broadcast();
+		OnStepsOrderChanged.Broadcast();
 		return true;
 	}
 
@@ -392,7 +392,7 @@ bool UDataprepActionAsset::SwapSteps(int32 FirstIndex, int32 SecondIndex)
 		UDataprepActionStep* FirstStep = Steps[FirstIndex];
 		Steps[FirstIndex] = Steps[SecondIndex];
 		Steps[SecondIndex] = FirstStep;
-		OnStepsChanged.Broadcast();
+		OnStepsOrderChanged.Broadcast();
 		return true;
 	}
 
@@ -438,8 +438,9 @@ bool UDataprepActionAsset::RemoveStep(int32 Index)
 			}
 		}
 
+		OnStepsAboutToBeRemoved.Broadcast( Steps[Index]->GetStepObject() );
 		Steps.RemoveAt( Index );
-		OnStepsChanged.Broadcast();
+		OnStepsOrderChanged.Broadcast();
 		return true;
 	}
 
@@ -449,7 +450,17 @@ bool UDataprepActionAsset::RemoveStep(int32 Index)
 
 FOnStepsOrderChanged& UDataprepActionAsset::GetOnStepsOrderChanged()
 {
-	return OnStepsChanged;
+	return OnStepsOrderChanged;
+}
+
+FOnStepAboutToBeRemoved& UDataprepActionAsset::GetOnStepAboutToBeRemoved()
+{
+	return OnStepsAboutToBeRemoved;
+}
+
+FOnStepWasEdited& UDataprepActionAsset::GetOnStepWasEdited()
+{
+	return OnStepWasEdited;
 }
 
 void UDataprepActionAsset::NotifyDataprepSystemsOfRemoval()
diff --git a/Engine/Plugins/Enterprise/DataprepEditor/Source/DataprepCore/Private/DataprepBindingCommandChange.h b/Engine/Plugins/Enterprise/DataprepEditor/Source/DataprepCore/Private/DataprepBindingCommandChange.h
new file mode 100644
--- /dev/null
+++ b/Engine/Plugins/Enterprise/DataprepEditor/Source/DataprepCore/Private/DataprepBindingCommandChange.h
@@ -0,0 +1,100 @@
+// Copyright Epic Games, Inc. All Rights Reserved.
+
+#pragma once
+
+#include "Misc/Change.h"
+
+template<class ObservedObjectType, typename DelatgateToBindType>
+class TDataprepBindingCommandChange final : public FCommandChange
+{
+public:
+	DECLARE_DELEGATE_RetVal(DelatgateToBindType&, FDelegateAccessor);
+	DECLARE_DELEGATE_TwoParams(OnUndoRedoDelegateType, ObservedObjectType&, const FDelegateHandle&);
+
+	using FDelegate = typename DelatgateToBindType::FDelegate;
+
+	TDataprepBindingCommandChange(bool bInWasBindingAdded
+		, ObservedObjectType& ObservedObject
+		, FDelegate&& InBinding
+		, FDelegateAccessor InDelegateAccessor
+		, OnUndoRedoDelegateType&& InOnApply
+		, OnUndoRedoDelegateType&& InOnRevert
+		, const FDelegateHandle& InHandle
+		);
+
+	virtual void Apply(UObject* Object) override;
+
+	virtual void Revert(UObject* Object) override;
+
+	virtual FString ToString() const override
+	{
+		return TEXT("Binding Tracking");
+	}
+
+private:
+	bool bWasBindingAdded;
+	TWeakObjectPtr<ObservedObjectType> ObservedObjectWeakPtr;
+	FDelegate Binding;
+	FDelegateAccessor DelegateAccessor;
+	OnUndoRedoDelegateType OnApply;
+	OnUndoRedoDelegateType OnRevert;
+	FDelegateHandle Handle;
+};
+
+template<class ObservedObjectType, typename DelatgateToBindType>
+TDataprepBindingCommandChange<ObservedObjectType, DelatgateToBindType>::TDataprepBindingCommandChange(
+	bool bInWasBindingAdded
+	, ObservedObjectType& ObservedObject
+	, FDelegate&& InBinding
+	, FDelegateAccessor InDelegateAccessor
+	, OnUndoRedoDelegateType&& InOnApply
+	, OnUndoRedoDelegateType&& InOnRevert
+	, const FDelegateHandle& InHandle
+	)
+	: bWasBindingAdded( bInWasBindingAdded )
+	, ObservedObjectWeakPtr( &ObservedObject )
+	, Binding( MoveTemp( InBinding ) )
+	, DelegateAccessor( MoveTemp( InDelegateAccessor ) )
+	, OnApply( MoveTemp( InOnApply ) )
+	, OnRevert( MoveTemp( InOnRevert ) )
+	, Handle( InHandle )
+{
+}
+
+template<class ObservedObjectType, typename DelatgateToBindType>
+void TDataprepBindingCommandChange<ObservedObjectType, DelatgateToBindType>::Apply(UObject* Object)
+{
+	ObservedObjectType* ObservedObject = ObservedObjectWeakPtr.Get();
+	if ( ObservedObject && DelegateAccessor.IsBound() )
+	{
+		DelatgateToBindType& Delegate = DelegateAccessor.Execute();
+		if ( bWasBindingAdded )
+		{
+			Handle = Delegate.Add( Binding );
+		}
+		else
+		{
+			Delegate.Remove( Handle );
+		}
+		OnApply.ExecuteIfBound( *ObservedObject, Handle );
+	}
+}
+
+template<class ObservedObjectType, typename DelatgateToBindType>
+void TDataprepBindingCommandChange<ObservedObjectType, DelatgateToBindType>::Revert(UObject* Object)
+{
+	ObservedObjectType* ObservedObject = ObservedObjectWeakPtr.Get();
+	if ( ObservedObject && DelegateAccessor.IsBound() )
+	{
+		DelatgateToBindType& Delegate = DelegateAccessor.Execute();
+		if ( bWasBindingAdded )
+		{
+			Delegate.Remove( Handle );
+		}
+		else
+		{
+			Handle = Delegate.Add( Binding );
+		}
+		OnRevert.ExecuteIfBound( *ObservedObject, Handle );
+	}
+}
diff --git a/Engine/Plugins/Enterprise/DataprepEditor/Source/DataprepCore/Private/DataprepCoreUtils.cpp b/Engine/Plugins/Enterprise/DataprepEditor/Source/DataprepCore/Private/DataprepCoreUtils.cpp
--- a/Engine/Plugins/Enterprise/DataprepEditor/Source/DataprepCore/Private/DataprepCoreUtils.cpp
+++ b/Engine/Plugins/Enterprise/DataprepEditor/Source/DataprepCore/Private/DataprepCoreUtils.cpp
@@ -6,6 +6,7 @@
 #include "DataprepAssetUserData.h"
 #endif
 
+#include "DataprepActionAsset.h"
 #include "DataprepAsset.h"
 #include "DataprepAssetInterface.h"
 #include "DataprepContentConsumer.h"
@@ -58,6 +59,20 @@ UDataprepAsset* FDataprepCoreUtils::GetDataprepAssetOfObject(UObject* Object)
 	return nullptr;
 }
 
+UDataprepActionAsset* FDataprepCoreUtils::GetDataprepActionAssetOf(UObject* Object)
+{
+	while (Object)
+	{
+		if (UDataprepActionAsset::StaticClass() == Object->GetClass())
+		{
+			return static_cast<UDataprepActionAsset*>(Object);
+		}
+		Object = Object->GetOuter();
+	}
+
+	return nullptr;
+}
+
 void FDataprepCoreUtils::PurgeObjects(TArray<UObject*> InObjects)
 {
 	TRACE_CPUPROFILER_EVENT_SCOPE(FDataprepCoreUtils::PurgeObjects);
diff --git a/Engine/Plugins/Enterprise/DataprepEditor/Source/DataprepCore/Private/DataprepParameterizableObject.cpp b/Engine/Plugins/Enterprise/DataprepEditor/Source/DataprepCore/Private/DataprepParameterizableObject.cpp
--- a/Engine/Plugins/Enterprise/DataprepEditor/Source/DataprepCore/Private/DataprepParameterizableObject.cpp
+++ b/Engine/Plugins/Enterprise/DataprepEditor/Source/DataprepCore/Private/DataprepParameterizableObject.cpp
@@ -2,38 +2,20 @@
 
 #include "DataprepParameterizableObject.h"
 
-#include "DataprepAsset.h"
-#include "DataprepCoreUtils.h"
-#include "Parameterization/DataprepParameterization.h"
-#include "Parameterization/DataprepParameterizationUtils.h"
+#include "DataprepBindingCommandChange.h"
 
+#include "CoreGlobals.h"
+#include "Misc/ITransaction.h"
+#include "Templates/UnrealTemplate.h"
 #include "UObject/UnrealType.h"
 
 void UDataprepParameterizableObject::PostEditChangeChainProperty(FPropertyChangedChainEvent& PropertyChangedEvent)
 {
 	Super::PostEditChangeChainProperty( PropertyChangedEvent );
 
-	if ( !bool( PropertyChangedEvent.ChangeType & ( EPropertyChangeType::Interactive | EPropertyChangeType::Redirected ) ) )
+	if ( !bool( PropertyChangedEvent.ChangeType & ( EPropertyChangeType::Interactive ) ) )
 	{
-		UDataprepAsset* DataprepAsset = FDataprepCoreUtils::GetDataprepAssetOfObject( this );
-		UDataprepParameterization * Parameterization = nullptr;
-		if ( DataprepAsset )
-		{
-			Parameterization = DataprepAsset->GetDataprepParameterization();
-		}
-
-		if ( !Parameterization )
-		{
-			UClass* Class = GetClass();
-			FString PathToDataprepParametrization = Class->GetMetaData( UDataprepParameterization::MetadataClassGeneratorName );
-			Parameterization = FindObject<UDataprepParameterization>( nullptr, *PathToDataprepParametrization );
-			
-		}
-
-		if ( Parameterization )
-		{
-			TArray<FDataprepPropertyLink> PropertyChain = FDataprepParameterizationUtils::MakePropertyChain( PropertyChangedEvent );
-			Parameterization->OnObjectPostEdit( this, PropertyChain, PropertyChangedEvent.ChangeType );
-		}
+		OnPostEdit.Broadcast( *this, PropertyChangedEvent );
 	}
 }
+
diff --git a/Engine/Plugins/Enterprise/DataprepEditor/Source/DataprepCore/Private/Parameterization/DataprepParameterization.cpp b/Engine/Plugins/Enterprise/DataprepEditor/Source/DataprepCore/Private/Parameterization/DataprepParameterization.cpp
--- a/Engine/Plugins/Enterprise/DataprepEditor/Source/DataprepCore/Private/Parameterization/DataprepParameterization.cpp
+++ b/Engine/Plugins/Enterprise/DataprepEditor/Source/DataprepCore/Private/Parameterization/DataprepParameterization.cpp
@@ -4,23 +4,27 @@
 
 #include "DataprepAsset.h"
 #include "DataprepAssetInstance.h"
+#include "DataprepBindingCommandChange.h"
+#include "DataprepCoreUtils.h"
 #include "DataprepParameterizableObject.h"
 #include "DataprepParameterizationArchive.h"
 #include "Parameterization/DataprepParameterizationUtils.h"
 
 #include "CoreGlobals.h"
 #include "DataprepCoreLogCategory.h"
+#include "Delegates/IDelegateInstance.h"
 #include "Editor.h"
 #include "Engine/Engine.h"
 #include "Math/UnrealMathUtility.h"
+#include "Misc/ITransaction.h"
 #include "Serialization/Archive.h"
 #include "Templates/UnrealTemplate.h"
+#include "UObject/CoreObjectVersion.h"
 #include "UObject/Object.h"
 #include "UObject/PropertyPortFlags.h"
 #include "UObject/UObjectGlobals.h"
 #include "UObject/UObjectHash.h"
 #include "UObject/UnrealType.h"
-#include "UObject/CoreObjectVersion.h"
 
 namespace DataprepParameterization
 {
@@ -572,6 +576,11 @@ bool UDataprepParameterizationBindings::HasBindingsForParameter(const FName& Par
 	return NameToBindings.Contains( ParameterName );
 }
 
+bool UDataprepParameterizationBindings::HasBindingsFromObject(UDataprepParameterizableObject* Object) const
+{
+	return ObjectToBindings.Contains( Object );
+}
+
 void UDataprepParameterizationBindings::Add(const TSharedRef<FDataprepParameterizationBinding>& Binding, const FName& ParamerterName, FSetOfBinding& OutBindingsContainedByNewBinding)
 {
 	Modify();
@@ -678,14 +687,14 @@ TSet<FName> UDataprepParameterizationBindings::RemoveAllBindingsFromObject(UData
 		Modify();
 
 		// Remove the bindings
-		TSet<FName> ParameterName;
+		TSet<FName> ParameterNames;
 		TArray<TSharedRef<FDataprepParameterizationBinding>> BindingsToRemove( Bindings->Array() );
 		for ( TSharedRef<FDataprepParameterizationBinding>& Binding : BindingsToRemove )
 		{
-			ParameterName.Add( RemoveBinding( Binding ) );
+			ParameterNames.Add( RemoveBinding( Binding ) );
 		}
 
-		return ParameterName;
+		return ParameterNames;
 	}
 
 	return {};
@@ -726,6 +735,13 @@ const UDataprepParameterizationBindings::FBindingToParameterNameMap& UDataprepPa
 	return BindingToParameterName;
 }
 
+TArray<UDataprepParameterizableObject*> UDataprepParameterizationBindings::GetParameterizedObjects() const
+{
+	TArray<UDataprepParameterizableObject*> Objects;
+	ObjectToBindings.GenerateKeyArray( Objects );
+	return Objects;
+}
+
 void UDataprepParameterizationBindings::Serialize(FArchive& Ar)
 {
 	Super::Serialize( Ar );
@@ -883,6 +899,18 @@ UDataprepParameterization::UDataprepParameterization()
 UDataprepParameterization::~UDataprepParameterization()
 {
 	FCoreUObjectDelegates::OnObjectModified.Remove( OnObjectModifiedHandle );
+
+	if ( BindingsContainer )
+	{
+		TArray<UDataprepParameterizableObject*> ParameterizedObjects = BindingsContainer->GetParameterizedObjects();
+		for ( UDataprepParameterizableObject* Object : ParameterizedObjects )
+		{
+			if ( Object )
+			{
+				RemoveBindingToPostEditOfParameterizableObject( *Object, false );
+			}
+		}
+	}
 }
 
 void UDataprepParameterization::PostInitProperties()
@@ -924,7 +952,16 @@ void UDataprepParameterization::PostLoad()
 		{
 			DoReinstancing( OldClass, false );
 			OnTellInstancesToReloadTheirSerializedData.Broadcast();
-			CastChecked<UDataprepAsset>( GetOuter() )->OnParameterizedObjectsChanged.Broadcast( nullptr );
+			CastChecked<UDataprepAsset>( GetOuter() )->OnParameterizedObjectsStatusChanged.Broadcast( nullptr );
+		}
+
+		TArray<UDataprepParameterizableObject*> ParameterizedObjects = BindingsContainer->GetParameterizedObjects();
+		for ( UDataprepParameterizableObject* Object : ParameterizedObjects )
+		{
+			if ( Object )
+			{
+				AddBindingToPostEditOfParameterizableObject( *Object, false );
+			}
 		}
 	}
 	Super::PostLoad();
@@ -942,22 +979,31 @@ void UDataprepParameterization::Serialize(FArchive& Ar)
 	Super::Serialize( Ar );
 }
 
-void UDataprepParameterization::PostEditUndo()
+void UDataprepParameterization::PostTransacted(const FTransactionObjectEvent& TransactionEvent)
 {
-	// This implementation work on the assumption that all the objects in the transaction were serialized before the calls to post edit undo
-	PrepareCustomClassForNewClassGeneration();
-	UClass* OldClass = CustomContainerClass;
-	CustomContainerClass = nullptr;
+	if ( TransactionEvent.GetChangedProperties().Contains( GET_MEMBER_NAME_CHECKED( UDataprepParameterization, ParameterizationStorage ) ) )
+	{
+		// This implementation work on the assumption that all the objects in the transaction were serialized before the calls to post edit undo
+		PrepareCustomClassForNewClassGeneration();
+		UClass* OldClass = CustomContainerClass;
+		CustomContainerClass = nullptr;
 
-	LoadParameterization();
+		LoadParameterization();
 
-	DoReinstancing( OldClass, false );
+		DoReinstancing( OldClass, false );
 
-	OnTellInstancesToReloadTheirSerializedData.Broadcast();
+		OnTellInstancesToReloadTheirSerializedData.Broadcast();
 
-	CastChecked<UDataprepAsset>( GetOuter() )->OnParameterizedObjectsChanged.Broadcast( nullptr );
+		CastChecked<UDataprepAsset>( GetOuter() )->OnParameterizedObjectsStatusChanged.Broadcast( nullptr );
+	}
+
+	if ( TransactionEvent.GetChangedProperties().Contains( GET_MEMBER_NAME_CHECKED( UDataprepParameterization, BindingsContainer ) ) )
+	{
+
+	}
 }
 
+
 void UDataprepParameterization::OnObjectModified(UObject* Object)
 {
 	if ( Object == DefaultParameterisation )
@@ -994,10 +1040,12 @@ bool UDataprepParameterization::BindObjectProperty(UDataprepParameterizableObjec
 
 		bool bBindingWasAdded = false;
 		bool bAddingFullProperty = PropertyChain.Last().ContainerIndex == INDEX_NONE;
+		bool bIsFirstBindingFromObject = false;
 
 		// We expect the chain to have a valid chain of cached property before inserting the binding
 		if ( DataprepParameterization::GetDeepestLevelOfValidCache( Binding.Get(), AddressOfTheValueFromBinding) == PropertyChain.Num() - 1 )
 		{
+			bIsFirstBindingFromObject = !BindingsContainer->HasBindingsFromObject( Object );
 			UDataprepParameterizationBindings::FSetOfBinding BindingsToRemove;
 
 			if ( FProperty** PropertyPtr = NameToParameterizationProperty.Find( Name ) )
@@ -1043,10 +1091,15 @@ bool UDataprepParameterization::BindObjectProperty(UDataprepParameterizableObjec
 		if ( bBindingWasAdded )
 		{
 			UpdateParameterizationFromBinding( Binding );
+			
+			if ( bIsFirstBindingFromObject )
+			{
+				AddBindingToPostEditOfParameterizableObject( *Object, true );
+			}
 
 			TSet<UObject*> Objects;
 			Objects.Add( Object );
-			CastChecked<UDataprepAsset>( GetOuter() )->OnParameterizedObjectsChanged.Broadcast( &Objects );
+			CastChecked<UDataprepAsset>( GetOuter() )->OnParameterizedObjectsStatusChanged.Broadcast( &Objects );
 		}
 	}
 
@@ -1081,7 +1134,7 @@ void UDataprepParameterization::RemoveBindedObjectProperty(UDataprepParameteriza
 
 		TSet<UObject*> Objects;
 		Objects.Add( Object );
-		CastChecked<UDataprepAsset>( GetOuter() )->OnParameterizedObjectsChanged.Broadcast( &Objects );
+		CastChecked<UDataprepAsset>( GetOuter() )->OnParameterizedObjectsStatusChanged.Broadcast( &Objects );
 	}
 }
 
@@ -1094,8 +1147,12 @@ void UDataprepParameterization::RemoveBindingFromObjects(TArray<UDataprepParamet
 
 	for ( UDataprepParameterizableObject* Object : Objects )
 	{
-		ParameterPotentiallyRemoved.Append( BindingsContainer->RemoveAllBindingsFromObject( Object ) );
-		UniqueObjects.Add( Object );
+		if ( Object )
+		{ 
+			ParameterPotentiallyRemoved.Append( BindingsContainer->RemoveAllBindingsFromObject( Object ) );
+			UniqueObjects.Add( Object );
+			RemoveBindingToPostEditOfParameterizableObject( *Object, true );
+		}
 	}
 
 	bool bClassWasChanged = false;
@@ -1113,7 +1170,7 @@ void UDataprepParameterization::RemoveBindingFromObjects(TArray<UDataprepParamet
 		UpdateClass();
 	}
 
-	CastChecked<UDataprepAsset>( GetOuter() )->OnParameterizedObjectsChanged.Broadcast( &UniqueObjects );
+	CastChecked<UDataprepAsset>( GetOuter() )->OnParameterizedObjectsStatusChanged.Broadcast( &UniqueObjects );
 }
 
 void UDataprepParameterization::UpdateParameterizationFromBinding(const TSharedRef<FDataprepParameterizationBinding>& Binding)
@@ -1139,28 +1196,7 @@ void UDataprepParameterization::UpdateParameterizationFromBinding(const TSharedR
 
 			TSet<UObject*> Objects;
 			Objects.Add( DefaultParameterisation );
-			CastChecked<UDataprepAsset>( GetOuter() )->OnParameterizedObjectsChanged.Broadcast( &Objects );
-		}
-	}
-}
-
-void UDataprepParameterization::OnObjectPostEdit(UDataprepParameterizableObject* Object, const TArray<FDataprepPropertyLink>& PropertyChain, EPropertyChangeType::Type ChangeType)
-{
-	if ( Object && PropertyChain.Num() > 0 )
-	{
-		if ( Object == DefaultParameterisation )
-		{
-			PushParametrizationValueToBindings( PropertyChain[0].PropertyName );
-		}
-		else
-		{
-			TSharedRef<FDataprepParameterizationBinding> BindingForModifiedProperty = MakeShared<FDataprepParameterizationBinding>( Object, PropertyChain );
-
-			TSharedPtr<FDataprepParameterizationBinding> Binding = BindingsContainer->GetContainingBinding( BindingForModifiedProperty );
-			if ( Binding )
-			{
-				UpdateParameterizationFromBinding( Binding.ToSharedRef() );
-			}
+			CastChecked<UDataprepAsset>( GetOuter() )->OnParameterizedObjectsStatusChanged.Broadcast( &Objects );
 		}
 	}
 }
@@ -1231,7 +1267,8 @@ void UDataprepParameterization::GenerateClass()
 		CustomContainerClass->StaticLink( true );
 		CustomContainerClass->AssembleReferenceTokenStream( true );
 
-		DefaultParameterisation = CustomContainerClass->GetDefaultObject( true );
+		DefaultParameterisation = static_cast<UDataprepParameterizableObject*>( CustomContainerClass->GetDefaultObject( true ) );
+		OnParameterizationDefaultObjectPostEditHandle = DefaultParameterisation->GetOnPostEdit().AddUObject( this, &UDataprepParameterization::OnParameterizationDefaultObjectPostEdit );
 	}
 }
 
@@ -1301,12 +1338,13 @@ void UDataprepParameterization::LoadParameterization()
 		CustomContainerClass->StaticLink(true);
 		CustomContainerClass->AssembleReferenceTokenStream(true);
 
-		DefaultParameterisation = CustomContainerClass->GetDefaultObject(true);
+		DefaultParameterisation = static_cast<UDataprepParameterizableObject*>( CustomContainerClass->GetDefaultObject(true) );
 		FDataprepParameterizationReader Reader( DefaultParameterisation, ParameterizationStorage );
+		OnParameterizationDefaultObjectPostEditHandle = DefaultParameterisation->GetOnPostEdit().AddUObject( this, &UDataprepParameterization::OnParameterizationDefaultObjectPostEdit );
 
 		if ( ObjectsToNotify.Num() > 0 )
 		{
-			CastChecked<UDataprepAsset>( GetOuter() )->OnParameterizedObjectsChanged.Broadcast( &ObjectsToNotify );
+			CastChecked<UDataprepAsset>( GetOuter() )->OnParameterizedObjectsStatusChanged.Broadcast( &ObjectsToNotify );
 		}
 	}
 }
@@ -1321,6 +1359,9 @@ void UDataprepParameterization::PrepareCustomClassForNewClassGeneration()
 		CustomContainerClass->Rename( *OldClassName, GetTransientPackage(), REN_DontCreateRedirectors | REN_DoNotDirty );
 		CustomContainerClass->SetFlags( RF_Transient | RF_NewerVersionExists );
 		CustomContainerClass->GetDefaultObject()->SetFlags( RF_Transient | RF_NewerVersionExists );
+
+		// Undo the hook for the post edit
+		DefaultParameterisation->GetOnPostEdit().Remove( OnParameterizationDefaultObjectPostEditHandle );
 	}
 }
 
@@ -1330,7 +1371,6 @@ void UDataprepParameterization::CreateClassObject()
 
 	CustomContainerClass = NewObject<UClass>(GetOutermost(), FName(TEXT("Parameterization")), RF_Transient);
 	CustomContainerClass->SetSuperStruct( UDataprepParameterizableObject::StaticClass() );
-	CustomContainerClass->SetMetaData( MetadataClassGeneratorName, *GetPathName() );
 }
 
 void UDataprepParameterization::DoReinstancing(UClass* OldClass, bool bMigrateData)
@@ -1388,8 +1428,6 @@ void UDataprepParameterization::DoReinstancing(UClass* OldClass, bool bMigrateDa
 		}
 
 		OnCustomClassWasUpdated.Broadcast( OldToNew );
-
-		DefaultParameterisation = NewCDO;
 	}
 }
 
@@ -1420,6 +1458,40 @@ FProperty* UDataprepParameterization::AddPropertyToClass(const FName& Parameteri
 	return nullptr;
 }
 
+
+// Just a macro to insert the definition in the right scope but to also keep the concerns separated of making the property change chain event from the function PushParametrizationValueToBindings
+#define MAKEPROPERTYCHANGEDCHAINEVENT( BINDING ) \
+	check( BINDING.PropertyChain.Num() > 0 ); \
+	TArray<UObject*, TInlineAllocator<1>> CurrentObject; \
+	CurrentObject.Add( BINDING.ObjectBinded ); \
+	FPropertyChangedEvent PropertyChangedEvent( BINDING.PropertyChain.Last().CachedProperty.Get(), EPropertyChangeType::Redirected, CurrentObject ); \
+	\
+	FEditPropertyChain EditPropertyChain; \
+	TMap<FString, int32> PropertiesIndex; \
+	PropertiesIndex.Reserve( BINDING.PropertyChain.Num() ); \
+	for ( const FDataprepPropertyLink& PropertyLink : BINDING.PropertyChain ) \
+	{ \
+		FProperty* Property = PropertyLink.CachedProperty.Get(); \
+		check( Property ); \
+		EditPropertyChain.AddTail( Property ); \
+		PropertiesIndex.Add( Property->GetName(), PropertyLink.ContainerIndex ); \
+	} \
+	TArray<TMap<FString, int32>, TInlineAllocator<1>> ArrayIndexPerObject; \
+	ArrayIndexPerObject.Add( PropertiesIndex ); \
+	PropertyChangedEvent.SetArrayIndexPerObject( ArrayIndexPerObject ); \
+	\
+	TDoubleLinkedList<FProperty*>::TDoubleLinkedListNode* TailNode = EditPropertyChain.GetTail(); \
+	check( TailNode ); \
+	EditPropertyChain.SetActivePropertyNode( TailNode->GetValue() ); \
+	\
+	if ( TDoubleLinkedList<FProperty*>::TDoubleLinkedListNode* NodeBeforeTail = TailNode->GetPrevNode() ) \
+	{ \
+		EditPropertyChain.SetActiveMemberPropertyNode( NodeBeforeTail->GetValue() ); \
+		PropertyChangedEvent.SetActiveMemberProperty( NodeBeforeTail->GetValue() ); \
+	} \
+	\
+	FPropertyChangedChainEvent PropertyChangedChainEvent(EditPropertyChain, PropertyChangedEvent);
+
 void UDataprepParameterization::PushParametrizationValueToBindings(FName ParameterName)
 {
 	if ( FProperty* ParameterizationPropeterty = NameToParameterizationProperty.FindRef( ParameterName ) )
@@ -1438,9 +1510,13 @@ void UDataprepParameterization::PushParametrizationValueToBindings(FName Paramet
 					void* AddressOfBindingValue = nullptr;
 					if ( FProperty* BindingProperty = DataprepParameterization::GetPropertyFromBinding( Binding.Get(), AddressOfBindingValue ) )
 					{
+						UDataprepParameterizableObject* ObjectBinded = Binding->ObjectBinded;
+						MAKEPROPERTYCHANGEDCHAINEVENT( Binding.Get() )
+						ObjectBinded->PreEditChange( PropertyChangedChainEvent.PropertyChain );
 						ObjectsModified.Add( Binding->ObjectBinded );
-						Binding->ObjectBinded->Modify();
+						ObjectBinded->Modify();
 						DataprepParameterization::CopyValue( *BindingProperty, AddressOfBindingValue, *ParameterizationPropeterty, AddressOfParameterValue );
+						ObjectBinded->PostEditChangeChainProperty( PropertyChangedChainEvent );
 					}
 					else
 					{
@@ -1449,7 +1525,7 @@ void UDataprepParameterization::PushParametrizationValueToBindings(FName Paramet
 				}
 
 				
-				CastChecked<UDataprepAsset>( GetOuter() )->OnParameterizedObjectsChanged.Broadcast( &ObjectsModified );
+				CastChecked<UDataprepAsset>( GetOuter() )->OnParameterizedObjectsStatusChanged.Broadcast( &ObjectsModified );
 
 				TSet<UObject*> ObjectsToNotify;
 				bool bClassNeedUpdate = false;
@@ -1473,13 +1549,16 @@ void UDataprepParameterization::PushParametrizationValueToBindings(FName Paramet
 
 				if ( ObjectsToNotify.Num() > 0 )
 				{
-					CastChecked<UDataprepAsset>( GetOuter() )->OnParameterizedObjectsChanged.Broadcast( &ObjectsToNotify );
+					CastChecked<UDataprepAsset>( GetOuter() )->OnParameterizedObjectsStatusChanged.Broadcast( &ObjectsToNotify );
 				}
 			}
 		}
 	}
 }
 
+#undef MAKEPROPERTYCHANGEDCHAINEVENT
+
+
 bool UDataprepParameterization::RemoveBinding(const TSharedRef<FDataprepParameterizationBinding>& Binding, bool& bOutClassNeedUpdate)
 {
 	FName ParameterOfRemovedBinding = BindingsContainer->RemoveBinding( Binding );
@@ -1488,7 +1567,120 @@ bool UDataprepParameterization::RemoveBinding(const TSharedRef<FDataprepParamete
 		bOutClassNeedUpdate = true;
 		NameToParameterizationProperty.Remove(ParameterOfRemovedBinding);
 	}
-	return !ParameterOfRemovedBinding.IsNone();
+
+	bool bBindingWasRemoved = !ParameterOfRemovedBinding.IsNone();
+
+	if ( bBindingWasRemoved )
+	{ 
+		UDataprepParameterizableObject* Object = Binding->ObjectBinded;
+		if ( !BindingsContainer->HasBindingsFromObject( Object ) && Object )
+		{
+			RemoveBindingToPostEditOfParameterizableObject( *Object, true );
+		}
+	}
+
+	return bBindingWasRemoved;
+}
+
+void UDataprepParameterization::OnParameterizationDefaultObjectPostEdit(UDataprepParameterizableObject& Object, FPropertyChangedChainEvent& PropertyChangedChainEvent)
+{
+	// When the parameterization push a new value to it's binding 
+	if ( !bool( PropertyChangedChainEvent.ChangeType & EPropertyChangeType::Redirected ) )
+	{
+		TArray<FDataprepPropertyLink> PropertyChain = FDataprepParameterizationUtils::MakePropertyChain( PropertyChangedChainEvent );
+		if ( PropertyChain.Num() > 0 )
+		{
+			PushParametrizationValueToBindings( PropertyChain[0].PropertyName );
+		}
+	}
+}
+
+void UDataprepParameterization::OnParameterizedObjectPostEdit(UDataprepParameterizableObject& Object, FPropertyChangedChainEvent& PropertyChangedChainEvent)
+{
+	if ( !bool( PropertyChangedChainEvent.ChangeType & EPropertyChangeType::Redirected ) )
+	{
+		TArray<FDataprepPropertyLink> PropertyChain = FDataprepParameterizationUtils::MakePropertyChain( PropertyChangedChainEvent );
+		if (PropertyChain.Num() > 0)
+		{
+			TSharedRef<FDataprepParameterizationBinding> BindingForModifiedProperty = MakeShared<FDataprepParameterizationBinding>(&Object, PropertyChain);
+			TSharedPtr<FDataprepParameterizationBinding> Binding = BindingsContainer->GetContainingBinding(BindingForModifiedProperty);
+			if ( Binding )
+			{
+				UpdateParameterizationFromBinding(Binding.ToSharedRef());
+			}
+		}
+	}
+}
+
+void UDataprepParameterization::OnAddedBindingToPostEditOfParameterizableObject(UDataprepParameterizableObject& Object, const FDelegateHandle& Handle)
+{
+	ObservedObjects.Add( &Object, Handle );
+}
+
+void UDataprepParameterization::OnRemovedBindingToPostEditOfParameterizableObject(UDataprepParameterizableObject& Object, const FDelegateHandle& Handle)
+{
+	ObservedObjects.Remove( &Object );
+}
+
+void UDataprepParameterization::AddBindingToPostEditOfParameterizableObject(UDataprepParameterizableObject& Object, bool bShouldAddToTransaction)
+{
+	UDataprepParameterizableObject::FOnPostEdit& OnPostEdit = Object.GetOnPostEdit();
+
+	using FPostEditDelegate = UDataprepParameterizableObject::FOnPostEdit::FDelegate;
+	FPostEditDelegate PostEditDelegate = FPostEditDelegate::CreateUObject( this, &UDataprepParameterization::OnParameterizedObjectPostEdit );
+
+	FDelegateHandle Handle = OnPostEdit.Add( PostEditDelegate );
+	OnAddedBindingToPostEditOfParameterizableObject( Object, Handle );
+
+	using FBindingCommandChange = TDataprepBindingCommandChange<UDataprepParameterizableObject, UDataprepParameterizableObject::FOnPostEdit>;
+	using FOnBindingChange = FBindingCommandChange::OnUndoRedoDelegateType;
+
+	if ( GUndo && bShouldAddToTransaction )
+	{
+		constexpr bool bIsAddingBinding = true;
+		GUndo->StoreUndo( &Object
+			, MakeUnique<FBindingCommandChange>( bIsAddingBinding
+				, Object
+				, MoveTemp( PostEditDelegate )
+				, FBindingCommandChange::FDelegateAccessor::CreateUObject( &Object, &UDataprepParameterizableObject::GetOnPostEdit )
+				, FOnBindingChange::CreateUObject( this, &UDataprepParameterization::OnAddedBindingToPostEditOfParameterizableObject )
+				, FOnBindingChange::CreateUObject( this, &UDataprepParameterization::OnRemovedBindingToPostEditOfParameterizableObject )
+				, Handle
+				)
+			);
+	}
+}
+
+void UDataprepParameterization::RemoveBindingToPostEditOfParameterizableObject(UDataprepParameterizableObject& Object, bool bShouldAddToTransaction)
+{
+	if ( FDelegateHandle* HandlePtr = ObservedObjects.Find( &Object ) )
+	{
+		FDelegateHandle Handle = *HandlePtr;
+		UDataprepParameterizableObject::FOnPostEdit& OnPostEdit = Object.GetOnPostEdit();
+		OnPostEdit.Remove( Handle );
+		OnRemovedBindingToPostEditOfParameterizableObject( Object, Handle );
+
+		using FBindingCommandChange = TDataprepBindingCommandChange<UDataprepParameterizableObject, UDataprepParameterizableObject::FOnPostEdit>;
+		using FOnBindingChange = FBindingCommandChange::OnUndoRedoDelegateType;
+
+		if ( GUndo && bShouldAddToTransaction )
+		{
+			using FPostEditDelegate = UDataprepParameterizableObject::FOnPostEdit::FDelegate;
+			FPostEditDelegate PostEditDelegate = FPostEditDelegate::CreateUObject( this, &UDataprepParameterization::OnParameterizedObjectPostEdit );
+
+			constexpr bool bIsAddingBinding = false;
+			GUndo->StoreUndo( &Object
+				, MakeUnique<FBindingCommandChange>( bIsAddingBinding
+					, Object
+					, MoveTemp( PostEditDelegate )
+					, FBindingCommandChange::FDelegateAccessor::CreateUObject( &Object, &UDataprepParameterizableObject::GetOnPostEdit )
+					, FOnBindingChange::CreateUObject( this, &UDataprepParameterization::OnRemovedBindingToPostEditOfParameterizableObject )
+					, FOnBindingChange::CreateUObject( this, &UDataprepParameterization::OnAddedBindingToPostEditOfParameterizableObject )
+					, Handle
+					)
+				);
+		}
+	}
 }
 
 bool UDataprepParameterization::OnAssetRename(ERenameFlags Flags)
@@ -1502,8 +1694,6 @@ bool UDataprepParameterization::OnAssetRename(ERenameFlags Flags)
 	return true;
 }
 
-const FName UDataprepParameterization::MetadataClassGeneratorName( TEXT("DataprepCustomParameterizationGenerator") );
-
 
 UDataprepParameterizationInstance::UDataprepParameterizationInstance()
 {
@@ -1520,11 +1710,19 @@ void UDataprepParameterization::AddReferencedObjects(UObject* InThis, FReference
 	Super::AddReferencedObjects(InThis, Collector);
 	UDataprepParameterization* This = CastChecked<UDataprepParameterization>(InThis);
 
-	for (TPair<FName, FProperty*>& Pair : This->NameToParameterizationProperty)
+	for ( TPair<FName, FProperty*>& Pair : This->NameToParameterizationProperty )
 	{
-		if (Pair.Value)
+		if ( Pair.Value )
 		{
-			Pair.Value->AddReferencedObjects(Collector);
+			Pair.Value->AddReferencedObjects( Collector );
+		}
+	}
+
+	for ( TPair<UDataprepParameterizableObject*, FDelegateHandle>& Pair : This->ObservedObjects)
+	{
+		if ( Pair.Key )
+		{
+			Collector.AddReferencedObject( Pair.Key );
 		}
 	}
 }
@@ -1653,4 +1851,3 @@ void UDataprepParameterizationInstance::UndoSetupForCallbacksFromParameterizatio
 		SourceParameterization->OnTellInstancesToReloadTheirSerializedData.RemoveAll( this );
 	}
 }
-
diff --git a/Engine/Plugins/Enterprise/DataprepEditor/Source/DataprepCore/Private/Parameterization/DataprepParameterization.h b/Engine/Plugins/Enterprise/DataprepEditor/Source/DataprepCore/Private/Parameterization/DataprepParameterization.h
--- a/Engine/Plugins/Enterprise/DataprepEditor/Source/DataprepCore/Private/Parameterization/DataprepParameterization.h
+++ b/Engine/Plugins/Enterprise/DataprepEditor/Source/DataprepCore/Private/Parameterization/DataprepParameterization.h
@@ -15,9 +15,11 @@
 
 #include "DataprepParameterization.generated.h"
 
-class UDataprepParameterizableObject;
-class FProperty;
 class FFieldClass;
+class FProperty;
+class FTransactionObjectEvent;
+class UDataprepParameterizableObject;
+class UDataprepActionAsset;
 
 struct FValueTypeValidationData
 {
@@ -173,6 +175,11 @@ public:
 	 */
 	bool HasBindingsForParameter(const FName& ParameterName) const;
 
+	/**
+	 * Does the data structure has some bindings from the specified object
+	 */
+	bool HasBindingsFromObject(UDataprepParameterizableObject* Object) const;
+
 	/**
 	 * Add a binding and map it to the parameter
 	 */
@@ -197,6 +204,8 @@ public:
 
 	const FBindingToParameterNameMap& GetBindingToParameterName() const;
 
+	TArray<UDataprepParameterizableObject*> GetParameterizedObjects() const;
+
 	// UObject interface
 	virtual void Serialize(FArchive& Ar) override;
 	static void AddReferencedObjects(UObject* InThis, FReferenceCollector& Collector);
@@ -234,14 +243,14 @@ public:
 	GENERATED_BODY()
 
 	UDataprepParameterization();
-	~UDataprepParameterization();
+	virtual ~UDataprepParameterization();
 
 	// UObject interface
 	virtual void PostInitProperties() override;
 	virtual void PostLoad() override;
 	virtual void Serialize(FArchive& Ar) override;
 	static void AddReferencedObjects(UObject* InThis, FReferenceCollector& Collector);
-	virtual void PostEditUndo() override;
+	virtual void PostTransacted(const FTransactionObjectEvent& TransactionEvent) override;
 	// End of UObject interface
 
 	void OnObjectModified(UObject* Object);
@@ -258,8 +267,6 @@ public:
 
 	void RemoveBindingFromObjects(TArray<UDataprepParameterizableObject*> Objects);
 
-	void OnObjectPostEdit(UDataprepParameterizableObject* Object, const TArray<FDataprepPropertyLink>& PropertyChain, EPropertyChangeType::Type ChangeType);
-
 	void GetExistingParameterNamesForType(FProperty* Property, bool bIsDescribingFullProperty, TSet<FString>& OutValidExistingNames, TSet<FString>& OutInvalidNames) const;
 
 private:
@@ -317,6 +324,16 @@ private:
 	 */
 	bool RemoveBinding(const TSharedRef<FDataprepParameterizationBinding>& Binding, bool& bOutClassNeedUpdate);
 
+	/**
+	 * Functions for the tracking of the post edit change
+	 */
+	void OnParameterizationDefaultObjectPostEdit(UDataprepParameterizableObject& Object, FPropertyChangedChainEvent& PropertyChangedChainEvent);
+	void OnParameterizedObjectPostEdit(UDataprepParameterizableObject& Object, FPropertyChangedChainEvent& PropertyChangedChainEvent);
+	void OnAddedBindingToPostEditOfParameterizableObject(UDataprepParameterizableObject& Object, const FDelegateHandle& Handle);
+	void OnRemovedBindingToPostEditOfParameterizableObject(UDataprepParameterizableObject& Object, const FDelegateHandle& Handle);
+	void AddBindingToPostEditOfParameterizableObject(UDataprepParameterizableObject& Object, bool bShouldAddToTransaction);
+	void RemoveBindingToPostEditOfParameterizableObject(UDataprepParameterizableObject& Object, bool bShouldAddToTransaction);
+
 public:
 
 	/**
@@ -325,22 +342,26 @@ public:
 	 */
 	bool OnAssetRename(ERenameFlags Flags);
 
-
-	static const FName MetadataClassGeneratorName;
-
 private:
-
 	// The containers for the bindings
 	UPROPERTY()
 	UDataprepParameterizationBindings* BindingsContainer;
 
+	/**
+	* Mapping of the observed object to their delegate handle
+	* The name are relative to the dataprep asset
+	*/
+	TMap<UDataprepParameterizableObject*, FDelegateHandle> ObservedObjects;
+
+	//TMap<UDataprepActionAsset*, TPair<FDelegateHandle, uint32> ObservedAction;
+
 	TMap<FName, FProperty*> NameToParameterizationProperty;
 
 	UPROPERTY(Transient, NonTransactional)
 	UClass* CustomContainerClass;
 
 	UPROPERTY(Transient, NonTransactional)
-	UObject* DefaultParameterisation;
+	UDataprepParameterizableObject* DefaultParameterisation;
 
 	/** 
 	 * This is used only to store a serialization of the values of the parameterization since we can't save our custom container class
@@ -362,6 +383,8 @@ private:
 	friend class UDataprepParameterizationInstance;
 
 	FDelegateHandle OnObjectModifiedHandle;
+
+	FDelegateHandle OnParameterizationDefaultObjectPostEditHandle;
 };
 
 
@@ -372,7 +395,7 @@ public:
 	GENERATED_BODY()
 
 	UDataprepParameterizationInstance();
-	~UDataprepParameterizationInstance();
+	virtual ~UDataprepParameterizationInstance();
 
 	// UObject interface
 	virtual void PostLoad() override;
diff --git a/Engine/Plugins/Enterprise/DataprepEditor/Source/DataprepCore/Public/DataprepActionAsset.h b/Engine/Plugins/Enterprise/DataprepEditor/Source/DataprepCore/Public/DataprepActionAsset.h
--- a/Engine/Plugins/Enterprise/DataprepEditor/Source/DataprepCore/Public/DataprepActionAsset.h
+++ b/Engine/Plugins/Enterprise/DataprepEditor/Source/DataprepCore/Public/DataprepActionAsset.h
@@ -163,8 +163,10 @@ struct FDataprepActionContext
 	DataprepActionAsset::FActionsContextChangedFunc ContextChangedCallback;
 };
 
-// Delegates
-DECLARE_MULTICAST_DELEGATE(FOnStepsOrderChanged)
+class UDataprepActionAsset;
+DECLARE_EVENT(UDataprepActionAsset, FOnStepsOrderChanged);
+DECLARE_EVENT_OneParam(UDataprepActionAsset, FOnStepAboutToBeRemoved, UDataprepParameterizableObject* /** The step object */ );
+DECLARE_EVENT_TwoParams(UDataprepParameterizableObject, FOnStepWasEdited, UDataprepParameterizableObject* /** The step object of the step (The object edited might be a subobject of it) */, struct FPropertyChangedChainEvent&)
 
 UCLASS(Experimental)
 class DATAPREPCORE_API UDataprepActionAsset : public UObject
@@ -300,6 +302,18 @@ public:
 	 */
 	FOnStepsOrderChanged& GetOnStepsOrderChanged();
 
+	/**
+	 * Allow an observer to be notified when a step is about to be removed from the action
+	 * @return The event that will be broadcasted
+	 */
+	FOnStepAboutToBeRemoved& GetOnStepAboutToBeRemoved();
+
+	/**
+	 * Allow an observer to be notified when a step was edited (the step himself or a sub object of it)
+	 * @return The event that will be broadcasted
+	 */
+	FOnStepWasEdited& GetOnStepWasEdited();
+
 	UPROPERTY(Transient)
 	bool bExecutionInterrupted;
 
@@ -398,7 +412,13 @@ private:
 	TArray<UDataprepActionStep*> Steps;
 
 	/** Broadcasts any change to the stack of steps */
-	FOnStepsOrderChanged OnStepsChanged;
+	FOnStepsOrderChanged OnStepsOrderChanged;
+
+	/** Broadcast when a step is about to be removed */
+	FOnStepAboutToBeRemoved OnStepsAboutToBeRemoved;
+
+	/** Broacast when a subobject of a step or a step object was edited */
+	FOnStepWasEdited OnStepWasEdited;
 
 	FDelegateHandle OnAssetDeletedHandle;
 
diff --git a/Engine/Plugins/Enterprise/DataprepEditor/Source/DataprepCore/Public/DataprepAsset.h b/Engine/Plugins/Enterprise/DataprepEditor/Source/DataprepCore/Public/DataprepAsset.h
--- a/Engine/Plugins/Enterprise/DataprepEditor/Source/DataprepCore/Public/DataprepAsset.h
+++ b/Engine/Plugins/Enterprise/DataprepEditor/Source/DataprepCore/Public/DataprepAsset.h
@@ -159,26 +159,6 @@ public:
 		return DataprepRecipeBP;
 	}
 
-	// struct to restrict the access scope
-	struct FDataprepBlueprintChangeNotifier
-	{
-	private:
-		friend class FDataprepEditorUtils;
-
-		static void NotifyDataprepBlueprintChange(UDataprepAsset& DataprepAsset, UObject* ModifiedObject)
-		{
-			// The asset is not complete yet. Skip this change
-			if(!DataprepAsset.HasAnyFlags(RF_NeedLoad | RF_NeedPostLoad | RF_NeedPostLoadSubobjects))
-			{
-				if(UDataprepActionAsset* ActionAsset = Cast<UDataprepActionAsset>(ModifiedObject))
-				{
-					DataprepAsset.UpdateActions();
-				}
-				DataprepAsset.OnBlueprintChanged.Broadcast( ModifiedObject );
-			}
-		}
-	};
-
 	//Todo Change the signature of this function when the new graph is ready (Hack to avoid a refactoring)
 	UDataprepActionAsset* AddActionUsingBP(class UEdGraphNode* NewActionNode);
 
@@ -186,12 +166,6 @@ public:
 
 	void RemoveActionUsingBP(int32 Index);
 
-	/**
-	 * Allow an observer to be notified of an change in the pipeline
-	 * return The event that will be broadcasted when a object has receive a modification that might change the result of the pipeline
-	 */
-	DECLARE_EVENT_OneParam(UDataprepAsset, FOnDataprepBlueprintChange, UObject* /*The object that was modified*/)
-	FOnDataprepBlueprintChange& GetOnBlueprintChanged() { return OnBlueprintChanged; }
 	// end of temp code for nodes development
 #endif
 
@@ -204,8 +178,9 @@ public:
 	 * Note on the objects param: The parameterized objects that should refresh their ui. If nullptr all widgets that can display some info on the parameterization should be refreshed
 	 */
 	DECLARE_EVENT_OneParam(UDataprepParameterization, FDataprepParameterizationStatusForObjectsChanged, const TSet<UObject*>* /** Objects */)
-	FDataprepParameterizationStatusForObjectsChanged OnParameterizedObjectsChanged;
+	FDataprepParameterizationStatusForObjectsChanged OnParameterizedObjectsStatusChanged;
 
+	// Internal Use only (the current implementation might be subject to change)
 	virtual UObject* GetParameterizationObject() override;
 
 	void BindObjectPropertyToParameterization(UDataprepParameterizableObject* Object, const TArray<struct FDataprepPropertyLink>& InPropertyChain,const FName& Name);
@@ -218,6 +193,7 @@ public:
 
 	void GetExistingParameterNamesForType(FProperty* Property, bool bIsDescribingFullProperty, TSet<FString>& OutValidExistingNames, TSet<FString>& OutInvalidNames) const;
 
+	// Internal only for now
 	UDataprepParameterization* GetDataprepParameterization() { return Parameterization; }
 
 protected:
@@ -259,9 +235,4 @@ private:
 
 	int32 CachedActionCount;
 
-#ifndef NO_BLUEPRINT
-	/** Event broadcasted when object in the pipeline was modified (Only broadcasted on changes that can affect the result of execution) */
-	FOnDataprepBlueprintChange OnBlueprintChanged;
-	// end of temp code for nodes development
-#endif
 };
\ No newline at end of file
diff --git a/Engine/Plugins/Enterprise/DataprepEditor/Source/DataprepCore/Public/DataprepCoreUtils.h b/Engine/Plugins/Enterprise/DataprepEditor/Source/DataprepCore/Public/DataprepCoreUtils.h
--- a/Engine/Plugins/Enterprise/DataprepEditor/Source/DataprepCore/Public/DataprepCoreUtils.h
+++ b/Engine/Plugins/Enterprise/DataprepEditor/Source/DataprepCore/Public/DataprepCoreUtils.h
@@ -15,6 +15,7 @@
 #include "UObject/Package.h"
 
 class UDataprepAsset;
+class UDataprepActionAsset;
 class UDataprepAssetInterface;
 class UDataprepParameterizableObject;
 struct FScopedSlowTask;
@@ -30,6 +31,12 @@ public:
 	 */
 	static UDataprepAsset* GetDataprepAssetOfObject(UObject* Object);
 
+	/**
+	 * Return the dataprep action asset that own the object, if the object is part of a dataprep action asset
+	 * @return nullptr if the object is not a part of a dataprep action asset
+	 */
+	static UDataprepActionAsset* GetDataprepActionAssetOf(UObject* Object);
+
 	/** Delete the objects and do the manipulation required to safely delete the assets */
 	static void PurgeObjects(TArray<UObject*> Objects);
 
diff --git a/Engine/Plugins/Enterprise/DataprepEditor/Source/DataprepCore/Public/DataprepParameterizableObject.h b/Engine/Plugins/Enterprise/DataprepEditor/Source/DataprepCore/Public/DataprepParameterizableObject.h
--- a/Engine/Plugins/Enterprise/DataprepEditor/Source/DataprepCore/Public/DataprepParameterizableObject.h
+++ b/Engine/Plugins/Enterprise/DataprepEditor/Source/DataprepCore/Public/DataprepParameterizableObject.h
@@ -4,6 +4,7 @@
 
 #include "CoreMinimal.h"
 #include "UObject/Object.h"
+#include "Templates/Function.h"
 
 #include "DataprepParameterizableObject.generated.h"
 
@@ -12,6 +13,7 @@ struct FPropertyChangedChainEvent;
 /**
  * The base class of all the object that can interact with the dataprep parameterization
  * This include all the parameterizable object and the parameterization object itself
+ * Also all the object that can be place in a dataprep action derive from it
  */
 UCLASS()
 class DATAPREPCORE_API UDataprepParameterizableObject : public UObject
@@ -20,4 +22,19 @@ public:
 	GENERATED_BODY()
 
 	virtual void PostEditChangeChainProperty(FPropertyChangedChainEvent& PropertyChangedEvent) override;
+
+	/*
+	 * Broadcast when a object has been post edited
+	 * Note that it doesn't broadcast on the interactive events
+	 */
+	DECLARE_EVENT_TwoParams(UDataprepParameterizableObject, FOnPostEdit, UDataprepParameterizableObject&, FPropertyChangedChainEvent&);
+	FOnPostEdit& GetOnPostEdit()
+	{ 
+		return OnPostEdit;
+	}
+
+private:
+
+	FOnPostEdit OnPostEdit;
 };
+
diff --git a/Engine/Plugins/Enterprise/DataprepEditor/Source/DataprepEditor/Private/BlueprintNodes/K2Node_DataprepAction.cpp b/Engine/Plugins/Enterprise/DataprepEditor/Source/DataprepEditor/Private/BlueprintNodes/K2Node_DataprepAction.cpp
--- a/Engine/Plugins/Enterprise/DataprepEditor/Source/DataprepEditor/Private/BlueprintNodes/K2Node_DataprepAction.cpp
+++ b/Engine/Plugins/Enterprise/DataprepEditor/Source/DataprepEditor/Private/BlueprintNodes/K2Node_DataprepAction.cpp
@@ -5,7 +5,6 @@
 // Dataprep includes
 #include "DataprepAsset.h"
 #include "DataprepActionAsset.h"
-#include "DataprepEditorUtils.h"
 #include "Widgets/DataprepGraph/SGraphNodeK2DataprepAction.h"
 
 // Engine includes
@@ -128,11 +127,6 @@ void UK2Node_DataprepAction::DestroyNode()
 	Super::DestroyNode();
 }
 
-void UK2Node_DataprepAction::NodeConnectionListChanged()
-{
-	FDataprepEditorUtils::NotifySystemOfChangeInPipeline( GetDataprepAction() );
-}
-
 TSharedPtr<SGraphNode> UK2Node_DataprepAction::CreateVisualWidget()
 {
 	return SNew( SGraphNodeK2DataprepAction, this );
diff --git a/Engine/Plugins/Enterprise/DataprepEditor/Source/DataprepEditor/Private/BlueprintNodes/K2Node_DataprepAction.h b/Engine/Plugins/Enterprise/DataprepEditor/Source/DataprepEditor/Private/BlueprintNodes/K2Node_DataprepAction.h
--- a/Engine/Plugins/Enterprise/DataprepEditor/Source/DataprepEditor/Private/BlueprintNodes/K2Node_DataprepAction.h
+++ b/Engine/Plugins/Enterprise/DataprepEditor/Source/DataprepEditor/Private/BlueprintNodes/K2Node_DataprepAction.h
@@ -30,7 +30,6 @@ public:
 	virtual FText GetNodeTitle(ENodeTitleType::Type TitleType) const override;
 	virtual void OnRenameNode(const FString& NewName) override;
 	virtual void DestroyNode() override;
-	virtual void NodeConnectionListChanged() override;
 	virtual TSharedPtr<SGraphNode> CreateVisualWidget() override;
 	TSharedPtr<class INameValidatorInterface> MakeNameValidator() const override;
 	// End EdGraphNode interface
diff --git a/Engine/Plugins/Enterprise/DataprepEditor/Source/DataprepEditor/Private/DataprepEditorUtils.cpp b/Engine/Plugins/Enterprise/DataprepEditor/Source/DataprepEditor/Private/DataprepEditorUtils.cpp
--- a/Engine/Plugins/Enterprise/DataprepEditor/Source/DataprepEditor/Private/DataprepEditorUtils.cpp
+++ b/Engine/Plugins/Enterprise/DataprepEditor/Source/DataprepEditor/Private/DataprepEditorUtils.cpp
@@ -27,36 +27,6 @@ bool FDataprepParametrizationActionData::IsValid() const
 	return DataprepAsset && Object && PropertyChain.Num() > 0;
 }
 
-void FDataprepEditorUtils::NotifySystemOfChangeInPipeline(UObject* SourceObject)
-{
-	UBlueprint* Blueprint = nullptr;
-	UDataprepAsset* DataprepAsset = nullptr;
-	UObject* Object = SourceObject;
-	while ( Object )
-	{
-		UClass* Class = Object->GetClass();
-		if ( Class->IsChildOf<UBlueprint>() )
-		{
-			Blueprint = static_cast<UBlueprint*>( Object );
-		}
-		else if ( Class == UDataprepAsset::StaticClass() )
-		{
-			DataprepAsset = static_cast<UDataprepAsset*>( Object );
-			break;
-		}
-		Object = Object->GetOuter();
-	}
-
-	if ( DataprepAsset )
-	{
-		UDataprepAsset::FDataprepBlueprintChangeNotifier::NotifyDataprepBlueprintChange( *DataprepAsset, SourceObject );
-	}
-	else if ( Blueprint )
-	{
-		FBlueprintEditorUtils::MarkBlueprintAsModified( Blueprint );
-	}
-}
-
 void FDataprepEditorUtils::PopulateMenuForParameterization(FMenuBuilder& MenuBuilder, UDataprepAsset& DataprepAsset, UDataprepParameterizableObject& Object, const TArray<FDataprepPropertyLink>& PropertyChain)
 {
 	TSharedRef<FDataprepParametrizationActionData> ActionData = MakeShared<FDataprepParametrizationActionData>( DataprepAsset, Object, PropertyChain );
diff --git a/Engine/Plugins/Enterprise/DataprepEditor/Source/DataprepEditor/Private/DataprepEditorUtils.h b/Engine/Plugins/Enterprise/DataprepEditor/Source/DataprepEditor/Private/DataprepEditorUtils.h
--- a/Engine/Plugins/Enterprise/DataprepEditor/Source/DataprepEditor/Private/DataprepEditorUtils.h
+++ b/Engine/Plugins/Enterprise/DataprepEditor/Source/DataprepEditor/Private/DataprepEditorUtils.h
@@ -33,11 +33,6 @@ struct FDataprepParametrizationActionData : public FGCObject
 class FDataprepEditorUtils
 {
 public:
-	/**
-	 * Find the the owning dataprep asset of the source object and send a notification to the dataprep editor so it can react when its pipeline is modified
-	 */
-	static void NotifySystemOfChangeInPipeline(UObject* SourceObject);
-
 	/**
 	 * Populate a menu builder with the section made for the parameterization
 	 * @param DataprepAsset the asset that own the object
diff --git a/Engine/Plugins/Enterprise/DataprepEditor/Source/DataprepEditor/Private/DataprepGraph/DataprepGraphActionNode.cpp b/Engine/Plugins/Enterprise/DataprepEditor/Source/DataprepEditor/Private/DataprepGraph/DataprepGraphActionNode.cpp
--- a/Engine/Plugins/Enterprise/DataprepEditor/Source/DataprepEditor/Private/DataprepGraph/DataprepGraphActionNode.cpp
+++ b/Engine/Plugins/Enterprise/DataprepEditor/Source/DataprepEditor/Private/DataprepGraph/DataprepGraphActionNode.cpp
@@ -5,7 +5,6 @@
 // Dataprep includes
 #include "DataprepAsset.h"
 #include "DataprepActionAsset.h"
-#include "DataprepEditorUtils.h"
 #include "Widgets/DataprepGraph/SDataprepGraphActionNode.h"
 #include "Widgets/DataprepGraph/SDataprepGraphEditor.h"
 
@@ -81,11 +80,6 @@ void UDataprepGraphActionNode::DestroyNode()
 	Super::DestroyNode();
 }
 
-void UDataprepGraphActionNode::NodeConnectionListChanged()
-{
-	FDataprepEditorUtils::NotifySystemOfChangeInPipeline( DataprepActionAsset );
-}
-
 TSharedPtr<class INameValidatorInterface> UDataprepGraphActionNode::MakeNameValidator() const
 {
 	// The name doesn't matter
diff --git a/Engine/Plugins/Enterprise/DataprepEditor/Source/DataprepEditor/Private/DataprepGraph/DataprepGraphActionNode.h b/Engine/Plugins/Enterprise/DataprepEditor/Source/DataprepEditor/Private/DataprepGraph/DataprepGraphActionNode.h
--- a/Engine/Plugins/Enterprise/DataprepEditor/Source/DataprepEditor/Private/DataprepGraph/DataprepGraphActionNode.h
+++ b/Engine/Plugins/Enterprise/DataprepEditor/Source/DataprepEditor/Private/DataprepGraph/DataprepGraphActionNode.h
@@ -74,7 +74,6 @@ public:
 	virtual FText GetNodeTitle(ENodeTitleType::Type TitleType) const override;
 	virtual void OnRenameNode(const FString& NewName) override;
 	virtual void DestroyNode() override;
-	virtual void NodeConnectionListChanged() override;
 	TSharedPtr<class INameValidatorInterface> MakeNameValidator() const override;
 	// End EdGraphNode interface
 
diff --git a/Engine/Plugins/Enterprise/DataprepEditor/Source/DataprepEditor/Private/SchemaActions/DataprepFetcherMenuActionCollector.cpp b/Engine/Plugins/Enterprise/DataprepEditor/Source/DataprepEditor/Private/SchemaActions/DataprepFetcherMenuActionCollector.cpp
deleted file mode 100644
diff --git a/Engine/Plugins/Enterprise/DataprepEditor/Source/DataprepEditor/Private/SchemaActions/DataprepFetcherMenuActionCollector.h b/Engine/Plugins/Enterprise/DataprepEditor/Source/DataprepEditor/Private/SchemaActions/DataprepFetcherMenuActionCollector.h
deleted file mode 100644
diff --git a/Engine/Plugins/Enterprise/DataprepEditor/Source/DataprepEditor/Private/Widgets/DataprepAssetView.cpp b/Engine/Plugins/Enterprise/DataprepEditor/Source/DataprepEditor/Private/Widgets/DataprepAssetView.cpp
--- a/Engine/Plugins/Enterprise/DataprepEditor/Source/DataprepEditor/Private/Widgets/DataprepAssetView.cpp
+++ b/Engine/Plugins/Enterprise/DataprepEditor/Source/DataprepEditor/Private/Widgets/DataprepAssetView.cpp
@@ -381,7 +381,7 @@ TSharedRef<ITableRow> SDataprepAssetView::OnGenerateRowForCategoryTree( TSharedR
 			if( DataprepAssetInterfacePtr->IsA<UDataprepAsset>() )
 			{
 				TWeakObjectPtr<UDataprepAsset> DataprepAsset = static_cast<UDataprepAsset*>( DataprepAssetInterfacePtr.Get() );
-				OnParameterizationWasEdited = DataprepAsset->OnParameterizedObjectsChanged.AddLambda( [ParameterizationDetailsView, DataprepAsset](const TSet<UObject*>* Objects )
+				OnParameterizationWasEdited = DataprepAsset->OnParameterizedObjectsStatusChanged.AddLambda( [ParameterizationDetailsView, DataprepAsset](const TSet<UObject*>* Objects )
 				{
 					if( Objects && Objects->Contains(DataprepAsset->GetParameterizationObject()) )
 					{
@@ -603,7 +603,7 @@ SDataprepAssetView::~SDataprepAssetView()
 			if ( DataprepAssetInterfacePtr->IsA<UDataprepAsset>() )
 			{
 				TWeakObjectPtr<UDataprepAsset> DataprepAsset = static_cast<UDataprepAsset*>( DataprepAssetInterfacePtr.Get() );
-				DataprepAsset->OnParameterizedObjectsChanged.Remove( OnParameterizationWasEdited );
+				DataprepAsset->OnParameterizedObjectsStatusChanged.Remove( OnParameterizationWasEdited );
 			}
 		}
 	}
diff --git a/Engine/Plugins/Enterprise/DataprepEditor/Source/DataprepEditor/Private/Widgets/DataprepGraph/SDataprepActionSteps.cpp b/Engine/Plugins/Enterprise/DataprepEditor/Source/DataprepEditor/Private/Widgets/DataprepGraph/SDataprepActionSteps.cpp
--- a/Engine/Plugins/Enterprise/DataprepEditor/Source/DataprepEditor/Private/Widgets/DataprepGraph/SDataprepActionSteps.cpp
+++ b/Engine/Plugins/Enterprise/DataprepEditor/Source/DataprepEditor/Private/Widgets/DataprepGraph/SDataprepActionSteps.cpp
@@ -5,7 +5,6 @@
 #include "DataprepActionAsset.h"
 #include "DataprepCoreUtils.h"
 #include "DataprepEditorStyle.h"
-#include "DataprepEditorUtils.h"
 #include "DataprepOperation.h"
 #include "SchemaActions/DataprepAllMenuActionCollector.h"
 #include "SchemaActions/DataprepDragDropOp.h"
@@ -210,7 +209,6 @@ void SDataprepActionSteps::Construct(const FArguments& InArgs, UDataprepActionAs
 void SDataprepActionSteps::OnStepsOrderChanged()
 {
 	Refresh();
-	FDataprepEditorUtils::NotifySystemOfChangeInPipeline( DataprepActionPtr.Get() );
 }
 
 void SDataprepActionSteps::Refresh()
diff --git a/Engine/Plugins/Enterprise/DataprepEditor/Source/DataprepEditor/Private/Widgets/DataprepGraph/SDataprepBoolFilter.cpp b/Engine/Plugins/Enterprise/DataprepEditor/Source/DataprepEditor/Private/Widgets/DataprepGraph/SDataprepBoolFilter.cpp
--- a/Engine/Plugins/Enterprise/DataprepEditor/Source/DataprepEditor/Private/Widgets/DataprepGraph/SDataprepBoolFilter.cpp
+++ b/Engine/Plugins/Enterprise/DataprepEditor/Source/DataprepEditor/Private/Widgets/DataprepGraph/SDataprepBoolFilter.cpp
@@ -3,7 +3,6 @@
 #include "Widgets/DataprepGraph/SDataprepBoolFilter.h"
 
 #include "SelectionSystem/DataprepBoolFilter.h"
-#include "Widgets/DataprepGraph/SDataprepFetcherSelector.h"
 
 #include "Widgets/Layout/SBox.h"
 
diff --git a/Engine/Plugins/Enterprise/DataprepEditor/Source/DataprepEditor/Private/Widgets/DataprepGraph/SDataprepFetcherSelector.cpp b/Engine/Plugins/Enterprise/DataprepEditor/Source/DataprepEditor/Private/Widgets/DataprepGraph/SDataprepFetcherSelector.cpp
deleted file mode 100644
diff --git a/Engine/Plugins/Enterprise/DataprepEditor/Source/DataprepEditor/Private/Widgets/DataprepGraph/SDataprepFetcherSelector.h b/Engine/Plugins/Enterprise/DataprepEditor/Source/DataprepEditor/Private/Widgets/DataprepGraph/SDataprepFetcherSelector.h
deleted file mode 100644
diff --git a/Engine/Plugins/Enterprise/DataprepEditor/Source/DataprepEditor/Private/Widgets/DataprepGraph/SDataprepFilter.cpp b/Engine/Plugins/Enterprise/DataprepEditor/Source/DataprepEditor/Private/Widgets/DataprepGraph/SDataprepFilter.cpp
--- a/Engine/Plugins/Enterprise/DataprepEditor/Source/DataprepEditor/Private/Widgets/DataprepGraph/SDataprepFilter.cpp
+++ b/Engine/Plugins/Enterprise/DataprepEditor/Source/DataprepEditor/Private/Widgets/DataprepGraph/SDataprepFilter.cpp
@@ -3,7 +3,6 @@
 #include "Widgets/DataprepGraph/SDataprepFilter.h"
 
 #include "DataprepEditorStyle.h"
-#include "DataprepEditorUtils.h"
 #include "SelectionSystem/DataprepBoolFilter.h"
 #include "SelectionSystem/DataprepFilter.h"
 #include "SelectionSystem/DataprepFloatFilter.h"
@@ -150,7 +149,15 @@ void SDataprepFilter::InverseFilter()
 	{
 		FScopedTransaction Transaction( LOCTEXT("InverseFilterTransaction", "Inverse the filter") );
 		Filter->SetIsExcludingResult( !Filter->IsExcludingResult() );
-		FDataprepEditorUtils::NotifySystemOfChangeInPipeline( Filter );
+
+		FProperty* Property = Filter->GetClass()->FindPropertyByName( TEXT("bIsExcludingResult") );
+		check( Property );
+		FEditPropertyChain EditChain;
+		EditChain.AddHead( Property );
+		EditChain.SetActivePropertyNode( Property );
+		FPropertyChangedEvent EditPropertyChangeEvent( Property, EPropertyChangeType::ValueSet );
+		FPropertyChangedChainEvent EditChangeChainEvent( EditChain, EditPropertyChangeEvent );
+		Filter->PostEditChangeChainProperty( EditChangeChainEvent );
 	}
 }
 
diff --git a/Engine/Plugins/Enterprise/DataprepEditor/Source/DataprepEditor/Private/Widgets/DataprepGraph/SDataprepFloatFilter.cpp b/Engine/Plugins/Enterprise/DataprepEditor/Source/DataprepEditor/Private/Widgets/DataprepGraph/SDataprepFloatFilter.cpp
--- a/Engine/Plugins/Enterprise/DataprepEditor/Source/DataprepEditor/Private/Widgets/DataprepGraph/SDataprepFloatFilter.cpp
+++ b/Engine/Plugins/Enterprise/DataprepEditor/Source/DataprepEditor/Private/Widgets/DataprepGraph/SDataprepFloatFilter.cpp
@@ -5,11 +5,9 @@
 #include "DataprepAsset.h"
 #include "DataprepCoreUtils.h"
 #include "DataprepEditorUtils.h"
-#include "SchemaActions/DataprepFetcherMenuActionCollector.h"
 #include "SchemaActions/DataprepSchemaAction.h"
 #include "SelectionSystem/DataprepFloatFilter.h"
 #include "Widgets/DataprepGraph/DataprepActionWidgetsUtils.h"
-#include "Widgets/DataprepGraph/SDataprepFetcherSelector.h"
 #include "Widgets/DataprepWidgets.h"
 #include "Widgets/Parameterization/SDataprepParameterizationLinkIcon.h"
 
@@ -84,7 +82,7 @@ void SDataprepFloatFilter::Construct(const FArguments& InArgs, UDataprepFloatFil
 			ToleranceParameterizationActionData = MakeShared<FDataprepParametrizationActionData>( *DataprepAsset, InFilter, PropertyChain );
 		}
 
-		OnParameterizationStatusForObjectsChangedHandle = DataprepAsset->OnParameterizedObjectsChanged.AddSP( this, &SDataprepFloatFilter::OnParameterizationStatusForObjectsChanged );
+		OnParameterizationStatusForObjectsChangedHandle = DataprepAsset->OnParameterizedObjectsStatusChanged.AddSP( this, &SDataprepFloatFilter::OnParameterizationStatusForObjectsChanged );
 	}
 
 	UpdateVisualDisplay();
@@ -94,7 +92,7 @@ SDataprepFloatFilter::~SDataprepFloatFilter()
 {
 	if ( UDataprepAsset* DataprepAsset = FDataprepCoreUtils::GetDataprepAssetOfObject( Filter ) )
 	{
-		DataprepAsset->OnParameterizedObjectsChanged.Remove( OnParameterizationStatusForObjectsChangedHandle );
+		DataprepAsset->OnParameterizedObjectsStatusChanged.Remove( OnParameterizationStatusForObjectsChangedHandle );
 	}
 }
 
@@ -276,8 +274,6 @@ void SDataprepFloatFilter::OnSelectedCriteriaChanged(TSharedPtr<FListEntry> List
 		Filter->SetFloatMatchingCriteria( FloatMatchType );
 
 		SDataprepFloatFilterUtils::PostEditChainProperty( Filter, TEXT("FloatMatchingCriteria") );
-
-		FDataprepEditorUtils::NotifySystemOfChangeInPipeline( Filter );
 	}
 }
 
@@ -335,8 +331,6 @@ void SDataprepFloatFilter::OnEqualValueComitted(float NewEqualValue, ETextCommit
 		Filter->SetEqualValue( NewEqualValue );
 
 		SDataprepFloatFilterUtils::PostEditChainProperty( Filter, TEXT("EqualValue") );
-
-		FDataprepEditorUtils::NotifySystemOfChangeInPipeline( Filter );
 		OldEqualValue = NewEqualValue;
 	}
 }
@@ -383,7 +377,6 @@ void SDataprepFloatFilter::OnToleranceComitted(float NewTolerance, ETextCommit::
 		FScopedTransaction Transaction( LOCTEXT("ToleranceChangedTransaction", "Change the Tolerance") );
 		Filter->SetTolerance( NewTolerance );
 		SDataprepFloatFilterUtils::PostEditChainProperty( Filter, TEXT("Tolerance") );
-		FDataprepEditorUtils::NotifySystemOfChangeInPipeline( Filter );
 		OldTolerance = NewTolerance;
 	}
 }
diff --git a/Engine/Plugins/Enterprise/DataprepEditor/Source/DataprepEditor/Private/Widgets/DataprepGraph/SDataprepIntegerFilter.cpp b/Engine/Plugins/Enterprise/DataprepEditor/Source/DataprepEditor/Private/Widgets/DataprepGraph/SDataprepIntegerFilter.cpp
--- a/Engine/Plugins/Enterprise/DataprepEditor/Source/DataprepEditor/Private/Widgets/DataprepGraph/SDataprepIntegerFilter.cpp
+++ b/Engine/Plugins/Enterprise/DataprepEditor/Source/DataprepEditor/Private/Widgets/DataprepGraph/SDataprepIntegerFilter.cpp
@@ -5,11 +5,9 @@
 #include "DataprepAsset.h"
 #include "DataprepCoreUtils.h"
 #include "DataprepEditorUtils.h"
-#include "SchemaActions/DataprepFetcherMenuActionCollector.h"
 #include "SchemaActions/DataprepSchemaAction.h"
 #include "SelectionSystem/DataprepIntegerFilter.h"
 #include "Widgets/DataprepGraph/DataprepActionWidgetsUtils.h"
-#include "Widgets/DataprepGraph/SDataprepFetcherSelector.h"
 #include "Widgets/DataprepWidgets.h"
 #include "Widgets/Parameterization/SDataprepParameterizationLinkIcon.h"
 
@@ -94,7 +92,7 @@ void SDataprepIntegerFilter::Construct(const FArguments& InArgs, UDataprepIntege
 			ToValueParameterizationActionData = MakeShared<FDataprepParametrizationActionData>(*DataprepAsset, InFilter, PropertyChain);
 		}
 
-		OnParameterizationStatusForObjectsChangedHandle = DataprepAsset->OnParameterizedObjectsChanged.AddSP( this, &SDataprepIntegerFilter::OnParameterizationStatusForObjectsChanged );
+		OnParameterizationStatusForObjectsChangedHandle = DataprepAsset->OnParameterizedObjectsStatusChanged.AddSP( this, &SDataprepIntegerFilter::OnParameterizationStatusForObjectsChanged );
 	}
 
 	UpdateVisualDisplay();
@@ -104,7 +102,7 @@ SDataprepIntegerFilter::~SDataprepIntegerFilter()
 {
 	if ( UDataprepAsset* DataprepAsset = FDataprepCoreUtils::GetDataprepAssetOfObject( Filter ) )
 	{
-		DataprepAsset->OnParameterizedObjectsChanged.Remove( OnParameterizationStatusForObjectsChangedHandle );
+		DataprepAsset->OnParameterizedObjectsStatusChanged.Remove( OnParameterizationStatusForObjectsChangedHandle );
 	}
 }
 
@@ -337,8 +335,6 @@ void SDataprepIntegerFilter::OnSelectedCriteriaChanged(TSharedPtr<FListEntry> Li
 		Filter->SetIntegerMatchingCriteria( IntMatchType );
 
 		SDataprepIntegerFilterUtils::PostEditChainProperty( Filter, TEXT("IntegerMatchingCriteria") );
-
-		FDataprepEditorUtils::NotifySystemOfChangeInPipeline( Filter );
 	}
 }
 
@@ -396,8 +392,6 @@ void SDataprepIntegerFilter::OnEqualValueComitted(int NewEqualValue, ETextCommit
 		Filter->SetEqualValue( NewEqualValue );
 
 		SDataprepIntegerFilterUtils::PostEditChainProperty( Filter, TEXT("EqualValue") );
-
-		FDataprepEditorUtils::NotifySystemOfChangeInPipeline( Filter );
 		OldEqualValue = NewEqualValue;
 	}
 }
@@ -439,8 +433,6 @@ void SDataprepIntegerFilter::OnFromValueComitted(int NewFromValue, ETextCommit::
 		Filter->SetFromValue(NewFromValue);
 
 		SDataprepIntegerFilterUtils::PostEditChainProperty(Filter, TEXT("FromValue"));
-
-		FDataprepEditorUtils::NotifySystemOfChangeInPipeline(Filter);
 		OldFromValue = NewFromValue;
 	}
 }
@@ -482,8 +474,6 @@ void SDataprepIntegerFilter::OnToValueComitted(int NewToValue, ETextCommit::Type
 		Filter->SetToValue(NewToValue);
 
 		SDataprepIntegerFilterUtils::PostEditChainProperty(Filter, TEXT("ToValue"));
-
-		FDataprepEditorUtils::NotifySystemOfChangeInPipeline(Filter);
 		OldToValue = NewToValue;
 	}
 }
diff --git a/Engine/Plugins/Enterprise/DataprepEditor/Source/DataprepEditor/Private/Widgets/DataprepGraph/SDataprepStringFilter.cpp b/Engine/Plugins/Enterprise/DataprepEditor/Source/DataprepEditor/Private/Widgets/DataprepGraph/SDataprepStringFilter.cpp
--- a/Engine/Plugins/Enterprise/DataprepEditor/Source/DataprepEditor/Private/Widgets/DataprepGraph/SDataprepStringFilter.cpp
+++ b/Engine/Plugins/Enterprise/DataprepEditor/Source/DataprepEditor/Private/Widgets/DataprepGraph/SDataprepStringFilter.cpp
@@ -6,12 +6,10 @@
 #include "DataprepCoreUtils.h"
 #include "DataprepEditorUtils.h"
 #include "Parameterization/DataprepParameterizationUtils.h"
-#include "SchemaActions/DataprepFetcherMenuActionCollector.h"
 #include "SchemaActions/DataprepSchemaAction.h"
 #include "SelectionSystem/DataprepStringFilter.h"
 #include "SelectionSystem/DataprepStringsArrayFilter.h"
 #include "Widgets/DataprepGraph/DataprepActionWidgetsUtils.h"
-#include "Widgets/DataprepGraph/SDataprepFetcherSelector.h"
 #include "Widgets/DataprepWidgets.h"
 #include "Widgets/Parameterization/SDataprepParameterizationLinkIcon.h"
 
@@ -59,7 +57,7 @@ void SDataprepStringFilter<FilterType>::Construct(const FArguments& InArgs, Filt
 			UserStringParameterizationActionData = MakeShared<FDataprepParametrizationActionData>( *DataprepAsset, InFilter, PropertyChain);
 		}
 
-		OnParameterizationStatusForObjectsChangedHandle = DataprepAsset->OnParameterizedObjectsChanged.AddSP( this, &SDataprepStringFilter<FilterType>::OnParameterizationStatusForObjectsChanged );
+		OnParameterizationStatusForObjectsChangedHandle = DataprepAsset->OnParameterizedObjectsStatusChanged.AddSP( this, &SDataprepStringFilter<FilterType>::OnParameterizationStatusForObjectsChanged );
 	}
 
 	UpdateVisualDisplay();
@@ -70,7 +68,7 @@ SDataprepStringFilter<FilterType>::~SDataprepStringFilter()
 {
 	if ( UDataprepAsset* DataprepAsset = FDataprepCoreUtils::GetDataprepAssetOfObject( Filter ) )
 	{
-		DataprepAsset->OnParameterizedObjectsChanged.Remove( OnParameterizationStatusForObjectsChangedHandle );
+		DataprepAsset->OnParameterizedObjectsStatusChanged.Remove( OnParameterizationStatusForObjectsChangedHandle );
 	}
 }
 
@@ -235,8 +233,6 @@ void SDataprepStringFilter<FilterType>::OnSelectedCriteriaChanged(TSharedPtr<FLi
 		FPropertyChangedEvent EditPropertyChangeEvent( Property, EPropertyChangeType::ValueSet );
 		FPropertyChangedChainEvent EditChangeChainEvent( EditChain, EditPropertyChangeEvent );
 		Filter->PostEditChangeChainProperty( EditChangeChainEvent );
-
-		FDataprepEditorUtils::NotifySystemOfChangeInPipeline( Filter );
 	}
 }
 
@@ -282,7 +278,6 @@ void SDataprepStringFilter<FilterType>::OnUserStringComitted(const FText& NewTex
 		FPropertyChangedChainEvent EditChangeChainEvent( EditChain, EditPropertyChangeEvent );
 		Filter->PostEditChangeChainProperty( EditChangeChainEvent );
 
-		FDataprepEditorUtils::NotifySystemOfChangeInPipeline( Filter );
 		OldUserString = NewUserString;
 	}
 }
diff --git a/Engine/Plugins/Enterprise/DataprepEditor/Source/DataprepEditor/Private/Widgets/DataprepWidgets.cpp b/Engine/Plugins/Enterprise/DataprepEditor/Source/DataprepEditor/Private/Widgets/DataprepWidgets.cpp
--- a/Engine/Plugins/Enterprise/DataprepEditor/Source/DataprepEditor/Private/Widgets/DataprepWidgets.cpp
+++ b/Engine/Plugins/Enterprise/DataprepEditor/Source/DataprepEditor/Private/Widgets/DataprepWidgets.cpp
@@ -656,8 +656,6 @@ void SDataprepDetailsView::CreateDefaultWidget( int32 Index, TSharedPtr< SWidget
 
 void SDataprepDetailsView::OnPropertyChanged(const FPropertyChangedEvent& InEvent)
 {
-	FDataprepEditorUtils::NotifySystemOfChangeInPipeline( DetailedObject );
-
 	if( TrackedProperties.Contains( InEvent.Property ) )
 	{
 		ForceRefresh();
@@ -916,7 +914,7 @@ void SDataprepDetailsView::Construct()
 		UDataprepAsset* DataprepAsset = FDataprepParameterizationUtils::GetDataprepAssetForParameterization( DetailedObject );
 		if ( DataprepAsset )
 		{
-			OnDataprepParameterizationStatusForObjectsChangedHandle = DataprepAsset->OnParameterizedObjectsChanged.AddSP( this, &SDataprepDetailsView::OnDataprepParameterizationStatusForObjectsChanged );
+			OnDataprepParameterizationStatusForObjectsChangedHandle = DataprepAsset->OnParameterizedObjectsStatusChanged.AddSP( this, &SDataprepDetailsView::OnDataprepParameterizationStatusForObjectsChanged );
 		}
 
 		if ( DetailedObject->IsA<UDataprepParameterizableObject>() )
@@ -983,7 +981,7 @@ SDataprepDetailsView::~SDataprepDetailsView()
 
 	if ( UDataprepAsset* DataprepAsset = DataprepAssetForParameterization.Get() )
 	{
-		DataprepAsset->OnParameterizedObjectsChanged.Remove( OnDataprepParameterizationStatusForObjectsChangedHandle );
+		DataprepAsset->OnParameterizedObjectsStatusChanged.Remove( OnDataprepParameterizationStatusForObjectsChangedHandle );
 	}
 }
 
