Integrating CL #11194391 from //UE4/Dev-Enterprise

Implement new class UDataprepFilterTransform

diff --git a/Engine/Plugins/Enterprise/DataprepEditor/Source/DataprepCore/Private/SelectionSystem/DataprepSelectionTransform.cpp b/Engine/Plugins/Enterprise/DataprepEditor/Source/DataprepCore/Private/SelectionSystem/DataprepSelectionTransform.cpp
new file mode 100644
--- /dev/null
+++ b/Engine/Plugins/Enterprise/DataprepEditor/Source/DataprepCore/Private/SelectionSystem/DataprepSelectionTransform.cpp
@@ -0,0 +1,99 @@
+// Copyright Epic Games, Inc. All Rights Reserved.
+
+#include "SelectionSystem/DataprepSelectionTransform.h"
+
+#include "DataprepCoreLogCategory.h"
+#include "IDataprepLogger.h"
+#include "IDataprepProgressReporter.h"
+
+#define LOCTEXT_NAMESPACE "DataprepSelectionTransform"
+
+void UDataprepSelectionTransform::Execute(const TArray<UObject*>& InObjects, TArray<UObject*>& OutObjects)
+{
+	OnExecution(InObjects, OutObjects);
+}
+
+void UDataprepSelectionTransform::OnExecution_Implementation(const TArray<UObject*>& InObjects, TArray<UObject*>& OutObjects)
+{
+	UE_LOG(LogDataprepCore, Error, TEXT("Please override UDataprepSelectionTransform::OnExecution_Implementation."));
+}
+
+FText UDataprepSelectionTransform::GetDisplayTransformName_Implementation() const
+{
+	return this->GetClass()->GetDisplayNameText();
+}
+
+FText UDataprepSelectionTransform::GetTooltip_Implementation() const
+{
+	return this->GetClass()->GetToolTipText();
+}
+
+FText UDataprepSelectionTransform::GetCategory_Implementation() const
+{
+	return LOCTEXT("DefaultSelectionTransformCategory", "");
+}
+
+FText UDataprepSelectionTransform::GetAdditionalKeyword_Implementation() const
+{
+	return FText();
+}
+
+void UDataprepRecursiveSelectionTransform::ApplySelectionTransform(const TArray<UObject*>& InObjects, TArray<UObject*>& OutObjects)
+{
+	UE_LOG(LogDataprepCore, Error, TEXT("Please override UDataprepRecursiveSelectionTransform::ApplySelectionTransform."));
+}
+
+void UDataprepRecursiveSelectionTransform::OnExecution_Implementation(const TArray<UObject*>& InObjects, TArray<UObject*>& OutObjects)
+{
+	ApplySelectionTransform(InObjects, OutObjects);
+
+	if (AllowRecursionLevels > 0 && OutObjects.Num() > 0)
+	{
+		TArray<UObject*> NewInputObjects(OutObjects);
+		TArray<UObject*> NewOutputObjects;
+		TSet<UObject*> ResultSet(OutObjects);
+
+		// Sets used to check for breaking conditions
+		TSet<UObject*> PrevInputSet;
+		TSet<UObject*> NewOutputSet;
+		
+		for (int32 Level = AllowRecursionLevels; Level > 0; --Level)
+		{
+			ApplySelectionTransform(NewInputObjects, NewOutputObjects);
+
+			if (NewOutputObjects.Num() == 0)
+			{
+				break; // No more output
+			}
+
+			NewOutputSet.Empty(NewOutputObjects.Num());
+			NewOutputSet.Append(NewOutputObjects);
+
+			if (PrevInputSet.Difference(NewOutputSet).Num() == 0)
+			{
+				break; // Infinite recursion
+			}
+
+			const int32 PrevResultSetNum = ResultSet.Num();
+
+			ResultSet.Append(NewOutputObjects);
+
+			if (PrevResultSetNum == ResultSet.Num())
+			{
+				break; // Nothing new
+			}
+
+			// Cache input to be able to detect infinite recursion
+			PrevInputSet.Empty(NewInputObjects.Num());
+			PrevInputSet.Append(NewInputObjects);
+
+			NewInputObjects.Empty(NewOutputObjects.Num());
+			NewInputObjects.Append(NewOutputObjects);
+			NewOutputObjects.Empty();
+		}
+
+		OutObjects = ResultSet.Array();
+	}
+}
+
+#undef LOCTEXT_NAMESPACE
diff --git a/Engine/Plugins/Enterprise/DataprepEditor/Source/DataprepCore/Public/SelectionSystem/DataprepSelectionTransform.h b/Engine/Plugins/Enterprise/DataprepEditor/Source/DataprepCore/Public/SelectionSystem/DataprepSelectionTransform.h
new file mode 100644
--- /dev/null
+++ b/Engine/Plugins/Enterprise/DataprepEditor/Source/DataprepCore/Public/SelectionSystem/DataprepSelectionTransform.h
@@ -0,0 +1,99 @@
+// Copyright Epic Games, Inc. All Rights Reserved.
+
+#pragma once
+
+#include "DataprepParameterizableObject.h"
+
+#include "CoreMinimal.h"
+#include "Delegates/DelegateCombinations.h"
+#include "UObject/Object.h"
+#include "UObject/ObjectMacros.h"
+#include "UObject/TextProperty.h"
+
+#include "DataprepSelectionTransform.generated.h"
+
+UCLASS(Abstract, Blueprintable)
+class DATAPREPCORE_API UDataprepSelectionTransform : public UDataprepParameterizableObject
+{
+	GENERATED_BODY()
+
+public:
+
+	/**
+	 * Execute the transform
+	 * @param InObjects Input objects the transform will operate on
+	 * @param OutObjects Resulting objects after the transform was executed
+	 */
+	UFUNCTION(BlueprintCallable, Category = "Execution")
+	void Execute(const TArray<UObject*>& InObjects, TArray<UObject*>& OutObjects);
+
+protected:
+	
+	/**
+	 * This function is called when the transform is executed.
+	 * If your defining your transform in Blueprint or Python this is the function to override.
+	 * @param InObjects Input objects the transform will operate on
+	 * @param OutObjects Resulting objects after the transform was executed
+	 */
+	UFUNCTION(BlueprintNativeEvent)
+	void OnExecution(const TArray<UObject*>& InObjects, TArray<UObject*>& OutObjects);
+
+	/**
+	 * This function is the same has OnExcution, but it's the extension point for an transform defined in c++.
+	 * It will be called on the transform execution.
+	 * @param InObjects Input objects the transform will operate on
+	 * @param OutObjects Resulting objects after the transform was executed
+	 */
+	virtual void OnExecution_Implementation(const TArray<UObject*>& InObjects, TArray<UObject*>& OutObjects);
+
+public:
+	/** 
+	 * Allows to change the name of the transform for the ui if needed.
+	 */
+	UFUNCTION(BlueprintNativeEvent,  Category = "Display")
+	FText GetDisplayTransformName() const;
+
+	/**
+	 * Allows to change the tooltip of the transform for the ui if needed.
+	 */
+	UFUNCTION(BlueprintNativeEvent,  Category = "Display")
+	FText GetTooltip() const;
+
+	/**
+	 * Allows to change the category of the transform for the ui if needed.
+	 */
+	UFUNCTION(BlueprintNativeEvent,  Category = "Display")
+	FText GetCategory() const;
+
+	/**
+	 * Allows to add more keywords for when a user is searching for the fetcher in the ui.
+	 */
+	UFUNCTION(BlueprintNativeEvent,  Category = "Display|Search")
+	FText GetAdditionalKeyword() const;
+
+	virtual FText GetDisplayTransformName_Implementation() const;
+	virtual FText GetTooltip_Implementation() const;
+	virtual FText GetCategory_Implementation() const;
+	virtual FText GetAdditionalKeyword_Implementation() const;
+
+protected:
+	// Specifies if input objects that have matching type can be added to the result
+	UPROPERTY(EditAnywhere, Category = SelectionTransform)
+	bool bOutputCanIncludeInput = true;
+};
+
+UCLASS(Abstract, Blueprintable)
+class DATAPREPCORE_API UDataprepRecursiveSelectionTransform : public UDataprepSelectionTransform
+{
+	GENERATED_BODY()
+
+public:
+	virtual void OnExecution_Implementation(const TArray<UObject*>& InObjects, TArray<UObject*>& OutObjects) override;
+
+	virtual void ApplySelectionTransform(const TArray<UObject*>& InObjects, TArray<UObject*>& OutObjects);
+
+private:
+	// How many times is it allowed to apply another transform on the result
+	UPROPERTY(EditAnywhere, Category = SelectionTransform)
+	int32 AllowRecursionLevels = -1;
+};
