Integrating CL #11208904 from //UE4/Dev-Enterprise

Fixed crash caused by UObject being garbage collected
Restored notification when BP is updated. Any change in BP overrides the simplified graph and reciprocally.
Fixed some logic flaws in copy/paste/cut/delete operations

diff --git a/Engine/Plugins/Enterprise/DataprepEditor/Source/DataprepCore/Private/DataprepActionAsset.cpp b/Engine/Plugins/Enterprise/DataprepEditor/Source/DataprepCore/Private/DataprepActionAsset.cpp
--- a/Engine/Plugins/Enterprise/DataprepEditor/Source/DataprepCore/Private/DataprepActionAsset.cpp
+++ b/Engine/Plugins/Enterprise/DataprepEditor/Source/DataprepCore/Private/DataprepActionAsset.cpp
@@ -79,6 +79,19 @@ UDataprepActionAsset::~UDataprepActionAsset()
 #endif //WITH_EDITOR
 }
 
+void UDataprepActionAsset::PostTransacted(const FTransactionObjectEvent& TransactionEvent)
+{
+	Super::PostTransacted(TransactionEvent);
+
+	if (TransactionEvent.GetEventType() == ETransactionObjectEventType::UndoRedo)
+	{
+		if (TransactionEvent.GetChangedProperties().Contains(TEXT("Steps")))
+		{
+			OnStepsOrderChanged.Broadcast();
+		}
+	}
+}
+
 void UDataprepActionAsset::Execute(const TArray<UObject*>& InObjects)
 {
 	ContextPtr = MakeShareable( new FDataprepActionContext() );
@@ -184,6 +197,7 @@ int32 UDataprepActionAsset::AddStep(const UDataprepActionStep* InActionStep)
 	{
 		Modify();
 		UDataprepActionStep* ActionStep = DuplicateObject<UDataprepActionStep>( InActionStep, this );
+		ActionStep->SetFlags(EObjectFlags::RF_Transactional);
 		Steps.Add( ActionStep );
 		OnStepsOrderChanged.Broadcast();
 		return Steps.Num() - 1;
@@ -208,6 +222,7 @@ int32 UDataprepActionAsset::AddSteps(const TArray<const UDataprepActionStep*>& I
 			if(InActionStep)
 			{
 				UDataprepActionStep* ActionStep = DuplicateObject<UDataprepActionStep>( InActionStep, this );
+				ActionStep->SetFlags(EObjectFlags::RF_Transactional);
 				Steps.Add( ActionStep );
 			}
 		}
@@ -261,6 +276,7 @@ bool UDataprepActionAsset::InsertStep(const UDataprepActionStep* InActionStep, i
 	{
 		Modify();
 		UDataprepActionStep* ActionStep = DuplicateObject<UDataprepActionStep>( InActionStep, this);
+		ActionStep->SetFlags(EObjectFlags::RF_Transactional);
 		Steps.Insert( ActionStep, Index );
 		OnStepsOrderChanged.Broadcast();
 		return true;
@@ -292,6 +308,7 @@ bool UDataprepActionAsset::InsertSteps(const TArray<const UDataprepActionStep*>&
 			if(InActionStep)
 			{
 				UDataprepActionStep* ActionStep = DuplicateObject<UDataprepActionStep>( InActionStep, this );
+				ActionStep->SetFlags(EObjectFlags::RF_Transactional);
 				Steps.Insert( ActionStep, Index );
 			}
 		}
@@ -416,10 +433,38 @@ bool UDataprepActionAsset::RemoveStep(int32 Index)
 {
 	if ( Steps.IsValidIndex( Index ) )
 	{
-		Modify();
+		return RemoveSteps({ Index });
+	}
+
+	UE_LOG( LogDataprepCore, Error, TEXT("UDataprepActionAsset::RemoveStep: The Index is out of range") );
+	return false;
+}
+
+bool UDataprepActionAsset::RemoveSteps(const TArray<int32>& Indices)
+{
+	if(Indices.Num() == 0)
+	{
+		UE_LOG( LogDataprepCore, Error, TEXT("UDataprepActionAsset::RemoveSteps: Empty array") );
+		return false;
+	}
+
+	UDataprepAsset* DataprepAsset = FDataprepCoreUtils::GetDataprepAssetOfObject( this );
+	if(DataprepAsset == nullptr)
+	{
+		UE_LOG( LogDataprepCore, Error, TEXT("UDataprepActionAsset::RemoveSteps: No Dataprep asset") );
+		return false;
+	}
 
-		if ( UDataprepAsset* DataprepAsset = FDataprepCoreUtils::GetDataprepAssetOfObject( this ) )
+	bool bSuccesfulRemoval = false;
+
+	Modify();
+
+	for(int32 Index : Indices)
+	{
+		if ( Steps.IsValidIndex( Index ) )
 		{
+			bSuccesfulRemoval = true;
+
 			if ( UDataprepParameterization* Parameterization = DataprepAsset->GetDataprepParameterization() )
 			{
 				TArray< UObject* > Objects;
@@ -436,15 +481,19 @@ bool UDataprepActionAsset::RemoveStep(int32 Index)
 
 				Parameterization->RemoveBindingFromObjects( ParameterizableObjects );
 			}
+
+			OnStepsAboutToBeRemoved.Broadcast( Steps[Index]->GetStepObject() );
+			Steps.RemoveAt( Index );
 		}
+	}
 
-		OnStepsAboutToBeRemoved.Broadcast( Steps[Index]->GetStepObject() );
-		Steps.RemoveAt( Index );
+	if(bSuccesfulRemoval)
+	{
 		OnStepsOrderChanged.Broadcast();
 		return true;
 	}
 
-	UE_LOG( LogDataprepCore, Error, TEXT("UDataprepActionAsset::RemoveStep: The Index is out of range") );
+	UE_LOG( LogDataprepCore, Error, TEXT("UDataprepActionAsset::RemoveSteps: Invalid array of step indices") );
 	return false;
 }
 
diff --git a/Engine/Plugins/Enterprise/DataprepEditor/Source/DataprepCore/Private/DataprepAsset.cpp b/Engine/Plugins/Enterprise/DataprepEditor/Source/DataprepCore/Private/DataprepAsset.cpp
--- a/Engine/Plugins/Enterprise/DataprepEditor/Source/DataprepCore/Private/DataprepAsset.cpp
+++ b/Engine/Plugins/Enterprise/DataprepEditor/Source/DataprepCore/Private/DataprepAsset.cpp
@@ -101,7 +101,7 @@ void UDataprepAsset::PostLoad()
 				DataprepRecipeBP->MarkPackageDirty();
 			}
 
-			UpdateActions();
+			UpdateActions(false);
 			bMarkDirty = true;
 		}
 #endif
@@ -408,6 +408,7 @@ TArray<UDataprepActionAsset*> UDataprepAsset::GetCopyOfActions(TMap<UObject*,UOb
 		DuplicationParameter.CreatedObjects = &OutOriginalToCopy;
 
 		UDataprepActionAsset* CopyOfAction = static_cast<UDataprepActionAsset*>( StaticDuplicateObjectEx( DuplicationParameter ) );
+		CopyOfAction->SetFlags(EObjectFlags::RF_Transactional);
 		check( CopyOfAction );
 
 		OutOriginalToCopy.Add( ActionAsset, CopyOfAction );
@@ -474,10 +475,9 @@ void UDataprepAsset::OnDataprepBlueprintChanged( UBlueprint* InBlueprint )
 	}
 }
 
-void UDataprepAsset::UpdateActions()
+void UDataprepAsset::UpdateActions(bool bNotify)
 {
-	TSet<UDataprepActionAsset*> ActionSet;
-	ActionSet.Append(ActionAssets);
+	ActionAssets.Empty(ActionAssets.Num());
 
 	UEdGraphPin* NodeOutPin= StartNode->FindPin(UEdGraphSchema_K2::PN_Then, EGPD_Output);
 	if( NodeOutPin && NodeOutPin->LinkedTo.Num() > 0 )
@@ -495,16 +495,16 @@ void UDataprepAsset::UpdateActions()
 				break;
 			}
 			else
-				{
+			{
 				ActionNodesVisited.AddByHash( NodeHash, NextNode );
-				}
+			}
 
 			if(UK2Node_DataprepActionCore* ActionNode = Cast<UK2Node_DataprepActionCore>(NextNode))
 			{
 				if( UDataprepActionAsset* DataprepAction = ActionNode->GetDataprepAction() )
 				{
 					DataprepAction->Rename(nullptr, this, REN_DoNotDirty | REN_DontCreateRedirectors | REN_NonTransactional);
-					ActionSet.Add( DataprepAction );
+					ActionAssets.Add( DataprepAction );
 				}
 			}
 
@@ -528,7 +528,10 @@ void UDataprepAsset::UpdateActions()
 		}
 	}
 
-	ActionAssets = ActionSet.Array();
+	if(bNotify)
+	{
+		OnActionChanged.Broadcast(ActionAssets.Num() > 0 ? ActionAssets[0] : nullptr, FDataprepAssetChangeType::ActionAdded);
+	}
 
 	CachedActionCount = ActionAssets.Num();
 }
@@ -541,6 +544,7 @@ int32 UDataprepAsset::AddAction(const UDataprepActionAsset* InAction)
 		Modify();
 
 		UDataprepActionAsset* Action = DuplicateObject<UDataprepActionAsset>( InAction, this);
+		Action->SetFlags(EObjectFlags::RF_Transactional);
 		Action->SetLabel( InAction->GetLabel() );
 
 		ActionAssets.Add( Action );
@@ -571,6 +575,7 @@ int32 UDataprepAsset::AddActions(const TArray<const UDataprepActionAsset*>& InAc
 			if(InAction)
 			{
 				UDataprepActionAsset* Action = DuplicateObject<UDataprepActionAsset>( InAction, this);
+				Action->SetFlags(EObjectFlags::RF_Transactional);
 				Action->SetLabel( InAction->GetLabel() );
 
 				ActionAssets.Add( Action );
@@ -658,6 +663,7 @@ bool UDataprepAsset::InsertAction(const UDataprepActionAsset* InAction, int32 In
 		Modify();
 
 		UDataprepActionAsset* Action = DuplicateObject<UDataprepActionAsset>( InAction, this);
+		Action->SetFlags(EObjectFlags::RF_Transactional);
 		Action->SetLabel( InAction->GetLabel());
 
 		ActionAssets.Insert( Action, Index );
@@ -697,6 +703,7 @@ bool UDataprepAsset::InsertActions(const TArray<const UDataprepActionAsset*>& In
 			if(InAction)
 			{
 				UDataprepActionAsset* Action = DuplicateObject<UDataprepActionAsset>( InAction, this);
+				Action->SetFlags(EObjectFlags::RF_Transactional);
 				Action->SetLabel( InAction->GetLabel() );
 
 				ActionAssets.Insert( Action, InsertIndex );
diff --git a/Engine/Plugins/Enterprise/DataprepEditor/Source/DataprepCore/Public/DataprepActionAsset.h b/Engine/Plugins/Enterprise/DataprepEditor/Source/DataprepCore/Public/DataprepActionAsset.h
--- a/Engine/Plugins/Enterprise/DataprepEditor/Source/DataprepCore/Public/DataprepActionAsset.h
+++ b/Engine/Plugins/Enterprise/DataprepEditor/Source/DataprepCore/Public/DataprepActionAsset.h
@@ -179,6 +179,10 @@ public:
 
 	virtual ~UDataprepActionAsset();
 
+	// Begin UObject Interface
+	virtual void PostTransacted(const FTransactionObjectEvent& TransactionEvent) override;
+	// End UObject Interface
+
 	/**
 	 * Execute the action on a specific set of objects
 	 * @param Objects The objects on which the action will operate
@@ -296,6 +300,13 @@ public:
 	 */
 	bool RemoveStep(int32 Index);
 
+	/**
+	 * Remove an array of steps from the action
+	 * @param Indices Array of step indices to remove
+	 * @return True if at least one step has been removed
+	 */
+	bool RemoveSteps(const TArray<int32>& Indices);
+
 	/**
 	 * Allow an observer to be notified when the steps order changed that also include adding and removing steps
 	 * @return The delegate that will be broadcasted when the steps order changed
diff --git a/Engine/Plugins/Enterprise/DataprepEditor/Source/DataprepCore/Public/DataprepAsset.h b/Engine/Plugins/Enterprise/DataprepEditor/Source/DataprepCore/Public/DataprepAsset.h
--- a/Engine/Plugins/Enterprise/DataprepEditor/Source/DataprepCore/Public/DataprepAsset.h
+++ b/Engine/Plugins/Enterprise/DataprepEditor/Source/DataprepCore/Public/DataprepAsset.h
@@ -63,6 +63,11 @@ public:
 
 	const UDataprepActionAsset* GetAction(int32 Index) const;
 
+	int32 GetActionIndex(UDataprepActionAsset* ActionAsset) const
+	{
+		return ActionAssets.Find(ActionAsset);
+	}
+
 	/**
 	 * Add a copy of the action to the Dataprep asset
 	 * @param Action The action we want to duplicate in the Dataprep asset
@@ -219,7 +224,7 @@ protected:
 
 	// Temp code for the nodes development
 private:
-	void UpdateActions();
+	void UpdateActions(bool bNotify = true);
 
 private:
 	UPROPERTY()
diff --git a/Engine/Plugins/Enterprise/DataprepEditor/Source/DataprepEditor/Private/BlueprintNodes/K2Node_DataprepAction.cpp b/Engine/Plugins/Enterprise/DataprepEditor/Source/DataprepEditor/Private/BlueprintNodes/K2Node_DataprepAction.cpp
--- a/Engine/Plugins/Enterprise/DataprepEditor/Source/DataprepEditor/Private/BlueprintNodes/K2Node_DataprepAction.cpp
+++ b/Engine/Plugins/Enterprise/DataprepEditor/Source/DataprepEditor/Private/BlueprintNodes/K2Node_DataprepAction.cpp
@@ -5,6 +5,7 @@
 // Dataprep includes
 #include "DataprepAsset.h"
 #include "DataprepActionAsset.h"
+#include "DataprepEditorUtils.h"
 #include "Widgets/DataprepGraph/SGraphNodeK2DataprepAction.h"
 
 // Engine includes
@@ -127,6 +128,11 @@ void UK2Node_DataprepAction::DestroyNode()
 	Super::DestroyNode();
 }
 
+void UK2Node_DataprepAction::NodeConnectionListChanged()
+{
+	FDataprepEditorUtils::NotifySystemOfChangeInPipeline( GetDataprepAction() );
+}
+
 TSharedPtr<SGraphNode> UK2Node_DataprepAction::CreateVisualWidget()
 {
 	return SNew( SGraphNodeK2DataprepAction, this );
diff --git a/Engine/Plugins/Enterprise/DataprepEditor/Source/DataprepEditor/Private/BlueprintNodes/K2Node_DataprepAction.h b/Engine/Plugins/Enterprise/DataprepEditor/Source/DataprepEditor/Private/BlueprintNodes/K2Node_DataprepAction.h
--- a/Engine/Plugins/Enterprise/DataprepEditor/Source/DataprepEditor/Private/BlueprintNodes/K2Node_DataprepAction.h
+++ b/Engine/Plugins/Enterprise/DataprepEditor/Source/DataprepEditor/Private/BlueprintNodes/K2Node_DataprepAction.h
@@ -30,6 +30,7 @@ public:
 	virtual FText GetNodeTitle(ENodeTitleType::Type TitleType) const override;
 	virtual void OnRenameNode(const FString& NewName) override;
 	virtual void DestroyNode() override;
+	virtual void NodeConnectionListChanged() override;
 	virtual TSharedPtr<SGraphNode> CreateVisualWidget() override;
 	TSharedPtr<class INameValidatorInterface> MakeNameValidator() const override;
 	// End EdGraphNode interface
diff --git a/Engine/Plugins/Enterprise/DataprepEditor/Source/DataprepEditor/Private/DataprepEditorUtils.cpp b/Engine/Plugins/Enterprise/DataprepEditor/Source/DataprepEditor/Private/DataprepEditorUtils.cpp
--- a/Engine/Plugins/Enterprise/DataprepEditor/Source/DataprepEditor/Private/DataprepEditorUtils.cpp
+++ b/Engine/Plugins/Enterprise/DataprepEditor/Source/DataprepEditor/Private/DataprepEditorUtils.cpp
@@ -27,6 +27,41 @@ bool FDataprepParametrizationActionData::IsValid() const
 	return DataprepAsset && Object && PropertyChain.Num() > 0;
 }
 
+#ifndef NO_BLUEPRINT
+void FDataprepEditorUtils::NotifySystemOfChangeInPipeline(UObject* SourceObject)
+{
+	UBlueprint* Blueprint = nullptr;
+	UDataprepAsset* DataprepAsset = nullptr;
+	UObject* Object = SourceObject;
+	while ( Object )
+	{
+		UClass* Class = Object->GetClass();
+		if ( Class->IsChildOf<UBlueprint>() )
+		{
+			Blueprint = static_cast<UBlueprint*>( Object );
+		}
+		else if ( Class == UDataprepAsset::StaticClass() )
+		{
+			DataprepAsset = static_cast<UDataprepAsset*>( Object );
+			break;
+		}
+		Object = Object->GetOuter();
+	}
+
+	if ( DataprepAsset )
+	{
+		if(!DataprepAsset->HasAnyFlags(RF_NeedLoad | RF_NeedPostLoad | RF_NeedPostLoadSubobjects))
+		{
+			DataprepAsset->GetRecipeBP()->OnChanged().Broadcast(DataprepAsset->GetRecipeBP());
+		}
+	}
+	else if ( Blueprint )
+	{
+		FBlueprintEditorUtils::MarkBlueprintAsModified( Blueprint );
+	}
+}
+#endif
+
 void FDataprepEditorUtils::PopulateMenuForParameterization(FMenuBuilder& MenuBuilder, UDataprepAsset& DataprepAsset, UDataprepParameterizableObject& Object, const TArray<FDataprepPropertyLink>& PropertyChain)
 {
 	TSharedRef<FDataprepParametrizationActionData> ActionData = MakeShared<FDataprepParametrizationActionData>( DataprepAsset, Object, PropertyChain );
diff --git a/Engine/Plugins/Enterprise/DataprepEditor/Source/DataprepEditor/Private/DataprepEditorUtils.h b/Engine/Plugins/Enterprise/DataprepEditor/Source/DataprepEditor/Private/DataprepEditorUtils.h
--- a/Engine/Plugins/Enterprise/DataprepEditor/Source/DataprepEditor/Private/DataprepEditorUtils.h
+++ b/Engine/Plugins/Enterprise/DataprepEditor/Source/DataprepEditor/Private/DataprepEditorUtils.h
@@ -33,6 +33,13 @@ struct FDataprepParametrizationActionData : public FGCObject
 class FDataprepEditorUtils
 {
 public:
+#ifndef NO_BLUEPRINT
+	/**
+	 * Find the the owning dataprep asset of the source object and send a notification to the dataprep editor so it can react when its pipeline is modified
+	 */
+	static void NotifySystemOfChangeInPipeline(UObject* SourceObject);
+#endif
+
 	/**
 	 * Populate a menu builder with the section made for the parameterization
 	 * @param DataprepAsset the asset that own the object
diff --git a/Engine/Plugins/Enterprise/DataprepEditor/Source/DataprepEditor/Private/SchemaActions/DataprepDragDropOp.cpp b/Engine/Plugins/Enterprise/DataprepEditor/Source/DataprepEditor/Private/SchemaActions/DataprepDragDropOp.cpp
--- a/Engine/Plugins/Enterprise/DataprepEditor/Source/DataprepEditor/Private/SchemaActions/DataprepDragDropOp.cpp
+++ b/Engine/Plugins/Enterprise/DataprepEditor/Source/DataprepEditor/Private/SchemaActions/DataprepDragDropOp.cpp
@@ -5,6 +5,7 @@
 #include "BlueprintNodes/K2Node_DataprepAction.h"
 #include "DataprepActionAsset.h"
 #include "DataprepAsset.h"
+#include "DataprepCoreUtils.h"
 #include "DataprepEditorLogCategory.h"
 #include "DataprepGraph/DataprepGraph.h"
 #include "DataprepGraph/DataprepGraphActionNode.h"
@@ -440,7 +441,20 @@ FReply FDataprepDragDropOp::DoDropOnActionStep(UDataprepGraphActionStepNode* Tar
 						{
 							if(!bCopyRequested)
 							{
-								bTransactionSuccessful &= SourceActionAsset->RemoveStep( StepIndex );
+								UDataprepAsset* DataprepAsset = FDataprepCoreUtils::GetDataprepAssetOfObject(SourceActionAsset);
+								// If the last step is removed from the source action, remove the whole action from the Dataprep asset
+								if(SourceActionAsset->GetStepsCount() == 1 && DataprepAsset)
+								{
+									// Find index of source action 
+									int32 Index = DataprepAsset->GetActionIndex(SourceActionAsset);
+									ensure(Index != INDEX_NONE);
+
+									bTransactionSuccessful &= DataprepAsset->RemoveAction(Index);
+								}
+								else
+								{
+									bTransactionSuccessful &= SourceActionAsset->RemoveStep( StepIndex );
+								}
 							}
 
 							bTransactionSuccessful &= TargetActionAsset->InsertStep( SourceActionStepPtr.Get(), TargetActionStepNode->GetStepIndex() );
@@ -493,32 +507,52 @@ FReply FDataprepDragDropOp::DoDropOnTrack(UDataprepAsset* TargetDataprepAsset, i
 
 	if(DraggedSteps.Num() > 0)
 	{
+		FModifierKeysState ModifierKeyState = FSlateApplication::Get().GetModifierKeys();
+		bool bCopySteps = ModifierKeyState.IsControlDown() || ModifierKeyState.IsCommandDown();
+
 		TArray<const UDataprepActionStep*> Steps;
 		Steps.Reserve(DraggedSteps.Num());
 
+		TArray<TStrongObjectPtr<UDataprepActionStep>> StepsStrongPtr;
+		StepsStrongPtr.Reserve(DraggedSteps.Num());
+
+		TMap<UDataprepActionAsset*, TArray<int32>> ActionAssetMap;
+
 		for(FDraggedStepEntry& DraggedStepEntry : DraggedSteps)
 		{
 			if(DraggedStepEntry.Get<2>().IsValid())
 			{
 				Steps.Add(DraggedStepEntry.Get<2>().Get());
+
+				// If copy is not active, hold onto the steps and build list of steps to remove
+				if(!bCopySteps)
+				{
+					StepsStrongPtr.Emplace(DraggedStepEntry.Get<2>().Get());
+					ActionAssetMap.FindOrAdd(DraggedStepEntry.Get<0>().Get()).Add(DraggedStepEntry.Get<1>());
+				}
 			}
 		}
 
 		if(Steps.Num() > 0)
 		{
-			FModifierKeysState ModifierKeyState = FSlateApplication::Get().GetModifierKeys();
-			bool bOnlyOneRequested = ModifierKeyState.IsControlDown() || ModifierKeyState.IsCommandDown();
-
 			FScopedTransaction Transaction( LOCTEXT("DropOnTrack", "Add Dataprep Action") );
 			bool bTransactionSuccessful = true;
 
+			if(!bCopySteps)
+			{
+				for(TPair<UDataprepActionAsset*, TArray<int32>>& Entry : ActionAssetMap)
+				{
+					Entry.Key->RemoveSteps(Entry.Value);
+				}
+			}
+
 			if(InsertIndex >= 0 && InsertIndex < TargetDataprepAsset->GetActionCount())
 			{
-				TargetDataprepAsset->InsertActions(Steps, InsertIndex, bOnlyOneRequested);
+				TargetDataprepAsset->InsertActions(Steps, InsertIndex);
 			}
 			else
 			{
-				TargetDataprepAsset->AddActions(Steps, bOnlyOneRequested);
+				TargetDataprepAsset->AddActions(Steps);
 			}
 		}
 
diff --git a/Engine/Plugins/Enterprise/DataprepEditor/Source/DataprepEditor/Private/Widgets/DataprepGraph/SDataprepGraphActionNode.cpp b/Engine/Plugins/Enterprise/DataprepEditor/Source/DataprepEditor/Private/Widgets/DataprepGraph/SDataprepGraphActionNode.cpp
--- a/Engine/Plugins/Enterprise/DataprepEditor/Source/DataprepEditor/Private/Widgets/DataprepGraph/SDataprepGraphActionNode.cpp
+++ b/Engine/Plugins/Enterprise/DataprepEditor/Source/DataprepEditor/Private/Widgets/DataprepGraph/SDataprepGraphActionNode.cpp
@@ -122,6 +122,7 @@ public:
 
 			+SVerticalBox::Slot()
 			.AutoHeight()
+			.Padding(10.f, 0.f, 10.f, 0.f)
 			[
 				SNew(SBorder)
 				.BorderBackgroundColor( this, &SDataprepEmptyActionStepNode::GetBorderBackgroundColor )
@@ -428,6 +429,8 @@ void SDataprepGraphActionNode::PopulateActionStepListWidget()
 	UEdGraph* EdGraph = GraphNode->GetGraph();
 	const int32 StepsCount = DataprepAction->GetStepsCount();
 	const UClass* GraphActionStepNodeClass = UDataprepGraphActionStepNode::StaticClass();
+	
+	EdGraphStepNodes.Reset(StepsCount);
 
 	TSharedPtr<SDataprepGraphTrackNode> TrackNodePtr = ParentTrackNodePtr.Pin();
 
@@ -437,7 +440,8 @@ void SDataprepGraphActionNode::PopulateActionStepListWidget()
 
 	for ( int32 Index = 0; Index < StepsCount; ++Index )
 	{
-		UDataprepGraphActionStepNode* ActionStepNode = NewObject<UDataprepGraphActionStepNode>( EdGraph, GraphActionStepNodeClass, NAME_None, RF_Transactional );
+		EdGraphStepNodes.Emplace(NewObject<UDataprepGraphActionStepNode>( EdGraph, GraphActionStepNodeClass, NAME_None, RF_Transactional ));
+		UDataprepGraphActionStepNode* ActionStepNode = EdGraphStepNodes.Last().Get();
 
 		ActionStepNode->CreateNewGuid();
 		ActionStepNode->PostPlacedNewNode();
diff --git a/Engine/Plugins/Enterprise/DataprepEditor/Source/DataprepEditor/Private/Widgets/DataprepGraph/SDataprepGraphActionNode.h b/Engine/Plugins/Enterprise/DataprepEditor/Source/DataprepEditor/Private/Widgets/DataprepGraph/SDataprepGraphActionNode.h
--- a/Engine/Plugins/Enterprise/DataprepEditor/Source/DataprepEditor/Private/Widgets/DataprepGraph/SDataprepGraphActionNode.h
+++ b/Engine/Plugins/Enterprise/DataprepEditor/Source/DataprepEditor/Private/Widgets/DataprepGraph/SDataprepGraphActionNode.h
@@ -4,6 +4,7 @@
 
 #include "CoreMinimal.h"
 #include "SGraphNode.h"
+#include "UObject/StrongObjectPtr.h"
 #include "Widgets/DeclarativeSyntaxSupport.h"
 
 class SDataprepGraphTrackNode;
@@ -12,6 +13,7 @@ class SDataprepGraphActionStepNode;
 class SVerticalBox;
 class UDataprepActionAsset;
 class UDataprepGraphActionNode;
+class UDataprepGraphActionStepNode;
 
 /**
  * The SDataprepGraphActionNode class is the SGraphNode associated
@@ -90,4 +92,7 @@ private:
 
 	/** Index of insert widget to be highlighted */
 	int32 InsertIndex;
+
+	/** Array of strong pointers to the UEdGraphNodes created for the action's steps */
+	TArray<TStrongObjectPtr<UDataprepGraphActionStepNode>> EdGraphStepNodes;
 };
diff --git a/Engine/Plugins/Enterprise/DataprepEditor/Source/DataprepEditor/Private/Widgets/DataprepGraph/SDataprepGraphActionStepNode.cpp b/Engine/Plugins/Enterprise/DataprepEditor/Source/DataprepEditor/Private/Widgets/DataprepGraph/SDataprepGraphActionStepNode.cpp
--- a/Engine/Plugins/Enterprise/DataprepEditor/Source/DataprepEditor/Private/Widgets/DataprepGraph/SDataprepGraphActionStepNode.cpp
+++ b/Engine/Plugins/Enterprise/DataprepEditor/Source/DataprepEditor/Private/Widgets/DataprepGraph/SDataprepGraphActionStepNode.cpp
@@ -120,6 +120,7 @@ void SDataprepGraphActionStepNode::UpdateGraphNode()
 
 		+SVerticalBox::Slot()
 		.AutoHeight()
+		.Padding(10.f, 0.f, 10.f, 0.f)
 		[
 			SNew(SBorder)
 			.BorderBackgroundColor( this, &SDataprepGraphActionStepNode::GetDragAndDropColor )
diff --git a/Engine/Plugins/Enterprise/DataprepEditor/Source/DataprepEditor/Private/Widgets/DataprepGraph/SDataprepGraphEditor.cpp b/Engine/Plugins/Enterprise/DataprepEditor/Source/DataprepEditor/Private/Widgets/DataprepGraph/SDataprepGraphEditor.cpp
--- a/Engine/Plugins/Enterprise/DataprepEditor/Source/DataprepEditor/Private/Widgets/DataprepGraph/SDataprepGraphEditor.cpp
+++ b/Engine/Plugins/Enterprise/DataprepEditor/Source/DataprepEditor/Private/Widgets/DataprepGraph/SDataprepGraphEditor.cpp
@@ -131,6 +131,21 @@ void SDataprepGraphEditor::OnDataprepAssetActionChanged(UObject* InObject, FData
 		case FDataprepAssetChangeType::ActionAdded:
 		case FDataprepAssetChangeType::ActionRemoved:
 		{
+			TArray<UEdGraphNode*> ToDelete;
+			TArray<UEdGraphNode*>& Nodes = GetCurrentGraph()->Nodes;
+			for(UEdGraphNode* NodeObject : Nodes)
+			{
+				if (UDataprepGraphActionNode* ActionNode = Cast<UDataprepGraphActionNode>(NodeObject))
+				{
+					ToDelete.Add(NodeObject);
+				}
+			}
+
+			for(UEdGraphNode* NodeObject : ToDelete)
+			{
+				Nodes.Remove(NodeObject);
+			}
+
 			TrackGraphNodePtr.Reset();
 			bIsComplete = false;
 			NotifyGraphChanged();
@@ -422,9 +437,10 @@ bool SDataprepGraphEditor::CanRenameNode() const
 
 void SDataprepGraphEditor::DeleteSelectedNodes()
 {
-	const FGraphPanelSelectionSet SelectedNodes = GetSelectedNodes();
+	FGraphPanelSelectionSet SelectedNodes = GetSelectedNodes();
 
 	TArray<int32> ActionsToDelete;
+	TSet<UDataprepActionAsset*> ActionAssets;
 	for(UObject* NodeObject : SelectedNodes)
 	{
 		if (UDataprepGraphActionNode* ActionNode = Cast<UDataprepGraphActionNode>(NodeObject))
@@ -432,21 +448,81 @@ void SDataprepGraphEditor::DeleteSelectedNodes()
 			if (ActionNode->CanUserDeleteNode() && ActionNode->GetDataprepActionAsset())
 			{
 				ActionsToDelete.Add(ActionNode->GetExecutionOrder());
+				ActionAssets.Add(ActionNode->GetDataprepActionAsset());
 			}
 		}
 	}
 
+	UEdGraph* EdGraph = GetCurrentGraph();
+
+	TMap<UDataprepActionAsset*, TArray<int32>> StepsToDelete;
+	for(UObject* NodeObject : SelectedNodes)
+	{
+		if (UDataprepGraphActionStepNode* ActionStepNode = Cast<UDataprepGraphActionStepNode>(NodeObject))
+		{
+			UDataprepActionAsset* ActionAsset = ActionStepNode->GetDataprepActionAsset();
+			if (ActionStepNode->CanUserDeleteNode() && ActionAsset && !ActionAssets.Contains(ActionAsset))
+			{
+				TArray<int32>& ToDelete = StepsToDelete.FindOrAdd(ActionAsset);
+				ToDelete.Add(ActionStepNode->GetStepIndex());
+
+				// Delete action if all its steps are deleted
+				if(ActionAsset->GetStepsCount() == ToDelete.Num())
+				{
+					StepsToDelete.Remove(ActionAsset);
+					int32 Index = DataprepAssetPtr->GetActionIndex(ActionAsset);
+					ensure(Index != INDEX_NONE);
+					ActionsToDelete.Add(Index);
+
+					TArray<class UEdGraphNode*>& Nodes = EdGraph->Nodes;
+					for(Index = 0; Index < Nodes.Num(); ++Index)
+					{
+						if(UDataprepGraphActionNode* ActionNode = Cast<UDataprepGraphActionNode>(Nodes[Index]))
+						{
+							if(ActionNode->GetDataprepActionAsset() == ActionAsset)
+							{
+								SelectedNodes.Add(Nodes[Index]);
+								break;
+							}
+						}
+					}
+				}
+			}
+		}
+	}
+
+	FScopedTransaction Transaction(FGenericCommands::Get().Delete->GetDescription());
+	bool bTransactionSuccessful = true;
+
 	if(ActionsToDelete.Num() > 0)
 	{
-		FScopedTransaction Transaction(FGenericCommands::Get().Delete->GetDescription());
+		bTransactionSuccessful &= DataprepAssetPtr->RemoveActions(ActionsToDelete);
+	}
 
-		bool bTransactionSuccessful = DataprepAssetPtr->RemoveActions(ActionsToDelete);
+	if(StepsToDelete.Num() > 0)
+	{
+		for(TPair<UDataprepActionAsset*, TArray<int32>>& Entry : StepsToDelete)
+		{
+			if(UDataprepActionAsset* ActionAsset = Entry.Key)
+			{
+				bTransactionSuccessful &= ActionAsset->RemoveSteps(Entry.Value);
+			}
+		}
+	}
 
-		if (!bTransactionSuccessful)
+	TArray<UEdGraphNode*>& Nodes = EdGraph->Nodes;
+	for(UObject* NodeObject : SelectedNodes)
+	{
+		if (UDataprepGraphActionNode* ActionNode = Cast<UDataprepGraphActionNode>(NodeObject))
 		{
-			Transaction.Cancel();
+			Nodes.Remove(ActionNode);
 		}
 	}
+
+	if (!bTransactionSuccessful)
+	{
+		Transaction.Cancel();
+	}
 }
 
 bool SDataprepGraphEditor::CanDeleteNodes() const
@@ -456,8 +532,8 @@ bool SDataprepGraphEditor::CanDeleteNodes() const
 	for (UObject* NodeObject : SelectedNodes)
 	{
 		// If any nodes allow deleting, then do not disable the delete option
-		UEdGraphNode* Node = Cast<UDataprepGraphActionNode>(NodeObject);
-		if (Node->CanUserDeleteNode())
+		UEdGraphNode* Node = Cast<UEdGraphNode>(NodeObject);
+		if (Node && Node->CanUserDeleteNode())
 		{
 			return true;
 			break;
diff --git a/Engine/Plugins/Enterprise/DataprepEditor/Source/DataprepEditor/Private/Widgets/DataprepGraph/SDataprepGraphTrackNode.cpp b/Engine/Plugins/Enterprise/DataprepEditor/Source/DataprepEditor/Private/Widgets/DataprepGraph/SDataprepGraphTrackNode.cpp
--- a/Engine/Plugins/Enterprise/DataprepEditor/Source/DataprepEditor/Private/Widgets/DataprepGraph/SDataprepGraphTrackNode.cpp
+++ b/Engine/Plugins/Enterprise/DataprepEditor/Source/DataprepEditor/Private/Widgets/DataprepGraph/SDataprepGraphTrackNode.cpp
@@ -190,12 +190,14 @@ void SDataprepGraphTrackNode::UpdateGraphNode()
 		const int32 ActionsCount = DataprepAsset->GetActionCount();
 		ActionNodes.Empty(ActionsCount);
 		ActionNodes.SetNum(ActionsCount);
+		EdGraphActionNodes.Reset(ActionsCount);
 
-		for(int32 Index = 0; Index < DataprepAsset->GetActionCount(); ++Index)
+		for(int32 Index = 0; Index < ActionsCount; ++Index)
 		{
 			if(UDataprepActionAsset* ActionAsset = DataprepAsset->GetAction(Index))
 			{
-				UDataprepGraphActionNode* NewActionNode = NewObject<UDataprepGraphActionNode>( EdGraph, UDataprepGraphActionNode::StaticClass(), NAME_None, RF_Transactional );
+				EdGraphActionNodes.Emplace(NewObject<UDataprepGraphActionNode>( EdGraph, UDataprepGraphActionNode::StaticClass(), NAME_None, RF_Transactional ));
+				UDataprepGraphActionNode* NewActionNode = EdGraphActionNodes.Last().Get();
 
 				NewActionNode->CreateNewGuid();
 				NewActionNode->PostPlacedNewNode();
@@ -424,7 +426,6 @@ void SDataprepGraphTrackNode::OnStartNodeDrag(const TSharedRef<SDataprepGraphAct
 	LastDragScreenSpacePosition = FSlateApplication::Get().GetCursorPos();
 	const FGeometry& DesktopGeometry = GetOwnerPanel()->GetPersistentState().DesktopGeometry;
 	FVector2D DragLocalPosition = DesktopGeometry.AbsoluteToLocal( LastDragScreenSpacePosition );
-	DragOrdinate = DragLocalPosition.Y;
 
 	SGraphPanel* GraphPanel = OwnerGraphPanelPtr.Pin().Get();
 	ensure(GraphPanel);
diff --git a/Engine/Plugins/Enterprise/DataprepEditor/Source/DataprepEditor/Private/Widgets/DataprepGraph/SDataprepGraphTrackNode.h b/Engine/Plugins/Enterprise/DataprepEditor/Source/DataprepEditor/Private/Widgets/DataprepGraph/SDataprepGraphTrackNode.h
--- a/Engine/Plugins/Enterprise/DataprepEditor/Source/DataprepEditor/Private/Widgets/DataprepGraph/SDataprepGraphTrackNode.h
+++ b/Engine/Plugins/Enterprise/DataprepEditor/Source/DataprepEditor/Private/Widgets/DataprepGraph/SDataprepGraphTrackNode.h
@@ -14,6 +14,7 @@
 #include "GraphEditor.h"
 #include "GraphEditorActions.h"
 #include "SGraphNode.h"
+#include "UObject/StrongObjectPtr.h"
 #include "Widgets/DeclarativeSyntaxSupport.h"
 #include "Widgets/Layout/SConstraintCanvas.h"
 #include "Widgets/SWidget.h"
@@ -143,17 +144,8 @@ private:
 	/** Cached of the last position of the cursor as the drag is happening */
 	FVector2D LastDragScreenSpacePosition;
 
-	/** Cached ordinate of the cursor when the drag started */
-	float DragOrdinate;
-
-	/** Execution order of the dragged node when the drag started */
-	int32 OriginalOrder;
-
-	/** Execution order of the dragged node as the drag is happening */
-	int32 CurrentOrder;
-
-	/** Array tracking the new execution order of actions while a drag is happening */
-	TArray<int32> NewActionsOrder;
+	/** Array of strong pointers to the UEdGraphNodes created for the Dataprep asset's actions */
+	TArray<TStrongObjectPtr<UDataprepGraphActionNode>> EdGraphActionNodes;
 
 	friend SDataprepGraphTrackWidget;
 };
