Integrating CL #11210847 from //UE4/Dev-Enterprise

Restored parameterization entries to contextual menu on step nodes
Unified removal of steps in an action. The action is deleted if no step remains after removal

diff --git a/Engine/Plugins/Enterprise/DataprepEditor/Source/DataprepCore/Private/DataprepCoreUtils.cpp b/Engine/Plugins/Enterprise/DataprepEditor/Source/DataprepCore/Private/DataprepCoreUtils.cpp
--- a/Engine/Plugins/Enterprise/DataprepEditor/Source/DataprepCore/Private/DataprepCoreUtils.cpp
+++ b/Engine/Plugins/Enterprise/DataprepEditor/Source/DataprepCore/Private/DataprepCoreUtils.cpp
@@ -644,4 +644,23 @@ void FDataprepCoreUtils::BuildAssets(const TArray<TWeakObjectPtr<UObject>>& Asse
 	});
 }
 
+bool FDataprepCoreUtils::RemoveSteps(UDataprepActionAsset* ActionAsset, const TArray<int32>& Indices, int32& ActionIndex)
+{
+	ActionIndex = INDEX_NONE;
+
+	UDataprepAsset* DataprepAsset = FDataprepCoreUtils::GetDataprepAssetOfObject(ActionAsset);
+	check(DataprepAsset);
+
+	if(ActionAsset->GetStepsCount() == Indices.Num())
+	{
+		// Find index of source action 
+		ActionIndex = DataprepAsset->GetActionIndex(ActionAsset);
+		ensure(ActionIndex != INDEX_NONE);
+
+		return DataprepAsset->RemoveAction(ActionIndex);
+	}
+
+	return ActionAsset->RemoveSteps( Indices );
+}
+
 #undef LOCTEXT_NAMESPACE
\ No newline at end of file
diff --git a/Engine/Plugins/Enterprise/DataprepEditor/Source/DataprepCore/Public/DataprepCoreUtils.h b/Engine/Plugins/Enterprise/DataprepEditor/Source/DataprepCore/Public/DataprepCoreUtils.h
--- a/Engine/Plugins/Enterprise/DataprepEditor/Source/DataprepCore/Public/DataprepCoreUtils.h
+++ b/Engine/Plugins/Enterprise/DataprepEditor/Source/DataprepCore/Public/DataprepCoreUtils.h
@@ -107,6 +107,30 @@ public:
 	 */
 	static void BuildAssets(const TArray< TWeakObjectPtr<UObject> >& Assets, const TSharedPtr<IDataprepProgressReporter>& ProgressReporterPtr );
 
+	/**
+	 * Helper function to remove a step from an action.
+	 * This function will remove the action from the Dataprep asset owning it if this is the last step.
+	 * 
+	 * @param	ActionAsset		Action asset to perform the operation on
+	 * @param	Indices			Array of step's indices to remove.
+	 * @param	ActionIndex		Set to INDEX_NONE if the action was not removed from its owning Dataprep asset. Valid index otherwise.
+	 * @return  Returns true if the removal was successful, false otherwise
+	 */
+	static bool RemoveSteps(UDataprepActionAsset* ActionAsset, const TArray<int32>& Indices, int32& ActionIndex );
+
+	/**
+	 * Helper function to remove a step from an action.
+	 * This function will remove the action from the Dataprep asset owning it if this is the last step.
+	 * 
+	 * @param	ActionAsset		Action asset to perform the operation on
+	 * @param	Index			Index of the step to remove.
+	 * @param	ActionIndex		Set to INDEX_NONE if the action was not removed from its owning Dataprep asset. Valid index otherwise.
+	 */
+	static bool RemoveStep(UDataprepActionAsset* ActionAsset, int32 Index, int32& ActionIndex)
+	{
+		return RemoveSteps(ActionAsset, { Index }, ActionIndex);
+	}
+
 	class DATAPREPCORE_API FDataprepLogger : public IDataprepLogger
 	{
 	public:
diff --git a/Engine/Plugins/Enterprise/DataprepEditor/Source/DataprepEditor/Private/SchemaActions/DataprepDragDropOp.cpp b/Engine/Plugins/Enterprise/DataprepEditor/Source/DataprepEditor/Private/SchemaActions/DataprepDragDropOp.cpp
--- a/Engine/Plugins/Enterprise/DataprepEditor/Source/DataprepEditor/Private/SchemaActions/DataprepDragDropOp.cpp
+++ b/Engine/Plugins/Enterprise/DataprepEditor/Source/DataprepEditor/Private/SchemaActions/DataprepDragDropOp.cpp
@@ -441,20 +441,8 @@ FReply FDataprepDragDropOp::DoDropOnActionStep(UDataprepGraphActionStepNode* Tar
 						{
 							if(!bCopyRequested)
 							{
-								UDataprepAsset* DataprepAsset = FDataprepCoreUtils::GetDataprepAssetOfObject(SourceActionAsset);
-								// If the last step is removed from the source action, remove the whole action from the Dataprep asset
-								if(SourceActionAsset->GetStepsCount() == 1 && DataprepAsset)
-								{
-									// Find index of source action 
-									int32 Index = DataprepAsset->GetActionIndex(SourceActionAsset);
-									ensure(Index != INDEX_NONE);
-
-									bTransactionSuccessful &= DataprepAsset->RemoveAction(Index);
-								}
-								else
-								{
-									bTransactionSuccessful &= SourceActionAsset->RemoveStep( StepIndex );
-								}
+								int32 ActionIndex = INDEX_NONE;
+								bTransactionSuccessful &= FDataprepCoreUtils::RemoveStep(SourceActionAsset, StepIndex, ActionIndex);
 							}
 
 							bTransactionSuccessful &= TargetActionAsset->InsertStep( SourceActionStepPtr.Get(), TargetActionStepNode->GetStepIndex() );
@@ -540,9 +528,13 @@ FReply FDataprepDragDropOp::DoDropOnTrack(UDataprepAsset* TargetDataprepAsset, i
 
 			if(!bCopySteps)
 			{
+				int32 ActionIndex;
 				for(TPair<UDataprepActionAsset*, TArray<int32>>& Entry : ActionAssetMap)
 				{
-					Entry.Key->RemoveSteps(Entry.Value);
+					if(Entry.Key)
+					{
+						bTransactionSuccessful &= FDataprepCoreUtils::RemoveSteps(Entry.Key, Entry.Value, ActionIndex);
+					}
 				}
 			}
 
@@ -640,18 +632,19 @@ FReply FDataprepDragDropOp::DoDropOnActionAsset(UDataprepGraphActionNode* Target
 					{
 						if(!bCopyRequested)
 						{
-							SourceActionAsset->RemoveStep( StepIndex );
+							int32 ActionIndex = INDEX_NONE;
+							bTransactionSuccessful &= FDataprepCoreUtils::RemoveStep(SourceActionAsset, StepIndex, ActionIndex);
 						}
 
-						TargetActionAsset->AddStep( SourceActionStepPtr.Get() );
+						bTransactionSuccessful &= TargetActionAsset->AddStep( SourceActionStepPtr.Get() ) != INDEX_NONE;
 					}
 					else if(bCopyRequested)
 					{
-						TargetActionAsset->AddStep( SourceActionStepPtr.Get() );
+						bTransactionSuccessful &= TargetActionAsset->AddStep( SourceActionStepPtr.Get() ) != INDEX_NONE;
 					}
 					else
 					{
-						TargetActionAsset->MoveStep( StepIndex, TargetActionAsset->GetStepsCount() - 1 );
+						bTransactionSuccessful &= TargetActionAsset->MoveStep( StepIndex, TargetActionAsset->GetStepsCount() - 1 );
 					}
 				}
 
diff --git a/Engine/Plugins/Enterprise/DataprepEditor/Source/DataprepEditor/Private/Widgets/DataprepGraph/SDataprepActionBlock.cpp b/Engine/Plugins/Enterprise/DataprepEditor/Source/DataprepEditor/Private/Widgets/DataprepGraph/SDataprepActionBlock.cpp
--- a/Engine/Plugins/Enterprise/DataprepEditor/Source/DataprepEditor/Private/Widgets/DataprepGraph/SDataprepActionBlock.cpp
+++ b/Engine/Plugins/Enterprise/DataprepEditor/Source/DataprepEditor/Private/Widgets/DataprepGraph/SDataprepActionBlock.cpp
@@ -3,6 +3,7 @@
 #include "Widgets/DataprepGraph/SDataprepActionBlock.h"
 
 #include "DataprepActionAsset.h"
+#include "DataprepCoreUtils.h"
 #include "DataprepEditorStyle.h"
 #include "SchemaActions/DataprepDragDropOp.h"
 #include "SchemaActions/DataprepSchemaAction.h"
@@ -32,10 +33,11 @@ void SDataprepActionBlock::Construct(const FArguments& InArgs, const TSharedRef<
 {
 	DataprepActionContext = InDataprepActionContext;
 
-	bDragEnabled = true;
-
 	const float DefaultPadding = FDataprepEditorStyle::GetFloat( "DataprepAction.Padding" );
 	const FLinearColor OutlineColor = GetOutlineColor();
+#ifndef NO_BLUEPRINT
+	bIsSimplifiedGraph = false;
+#endif
 
 	ChildSlot
 	[
@@ -85,7 +87,7 @@ void SDataprepActionBlock::Construct(const FArguments& InArgs, const TSharedRef<
 			// The content zone of the action block
 			+ SVerticalBox::Slot()
 			.AutoHeight()
-			.Padding( bDragEnabled ? FMargin( DefaultPadding ) : FMargin( 2.f * DefaultPadding, 0.f, 0.f, 0.f ) )
+			.Padding( !bIsSimplifiedGraph ? FMargin( DefaultPadding ) : FMargin( 2.f * DefaultPadding, 0.f, 0.f, 0.f ) )
 			[
 				SNew( SConstraintCanvas )
 
@@ -144,15 +146,22 @@ TSharedRef<SWidget> SDataprepActionBlock::GetBlockTitleWidget()
 
 FReply SDataprepActionBlock::OnMouseButtonDown(const FGeometry& MyGeometry, const FPointerEvent& MouseEvent)
 {
-	if ( bDragEnabled && MouseEvent.GetEffectingButton() ==  EKeys::LeftMouseButton )
+#ifndef NO_BLUEPRINT
+	if ( !bIsSimplifiedGraph && MouseEvent.GetEffectingButton() ==  EKeys::LeftMouseButton )
 	{
 		return FReply::Handled().DetectDrag( AsShared(), EKeys::LeftMouseButton );
 	}
+#else
+	if ( MouseEvent.GetEffectingButton() ==  EKeys::LeftMouseButton )
+	{
+		return FReply::Handled().DetectDrag( AsShared(), EKeys::LeftMouseButton );
+	}
+#endif
 
 	// Take ownership of the mouse if right mouse button clicked to display contextual menu
 	if ( MouseEvent.GetEffectingButton() == EKeys::RightMouseButton )
 	{
-		return !bDragEnabled ? FReply::Handled().CaptureMouse( SharedThis(this) ) : FReply::Handled();
+		return FReply::Handled();
 	}
 
 	return FReply::Unhandled();
@@ -373,14 +382,18 @@ TSharedRef<SWidget> SDataprepActionBlock::GetContentBackgroundWidget()
 
 void SDataprepActionBlock::PopulateMenuBuilder(FMenuBuilder& MenuBuilder)
 {
-	FUIAction DeleteAction;
-	DeleteAction.ExecuteAction.BindSP( this, &SDataprepActionBlock::DeleteStep );
-
-	TSharedPtr<FUICommandInfo> DeleteCommand = FGenericCommands::Get().Delete;
-	MenuBuilder.AddMenuEntry( DeleteCommand->GetLabel(),
-		DeleteCommand->GetDescription(),
-		DeleteCommand->GetIcon(),
-		DeleteAction );
+	MenuBuilder.BeginSection( FName( TEXT("NodeSection") ), LOCTEXT("NodeSection", "Common") );
+	{
+		FUIAction DeleteAction;
+		DeleteAction.ExecuteAction.BindSP( this, &SDataprepActionBlock::DeleteStep );
+
+		TSharedPtr<FUICommandInfo> DeleteCommand = FGenericCommands::Get().Delete;
+		MenuBuilder.AddMenuEntry( DeleteCommand->GetLabel(),
+			DeleteCommand->GetDescription(),
+			DeleteCommand->GetIcon(),
+			DeleteAction );
+	}
+	MenuBuilder.EndSection();
 }
 
 void SDataprepActionBlock::DeleteStep()
@@ -390,7 +403,12 @@ void SDataprepActionBlock::DeleteStep()
 		if ( UDataprepActionAsset* ActionAsset = ActionContext->DataprepActionPtr.Get() )
 		{
 			FScopedTransaction Transaction( LOCTEXT("DeleteStepTransaction", "Remove step from action") );
-			ActionAsset->RemoveStep( ActionContext->StepIndex );
+
+			int32 ActionIndex = INDEX_NONE;
+			if( !FDataprepCoreUtils::RemoveStep( ActionAsset, ActionContext->StepIndex, ActionIndex ) )
+			{
+				Transaction.Cancel();
+			}
 		}
 	}
 }
diff --git a/Engine/Plugins/Enterprise/DataprepEditor/Source/DataprepEditor/Private/Widgets/DataprepGraph/SDataprepActionBlock.h b/Engine/Plugins/Enterprise/DataprepEditor/Source/DataprepEditor/Private/Widgets/DataprepGraph/SDataprepActionBlock.h
--- a/Engine/Plugins/Enterprise/DataprepEditor/Source/DataprepEditor/Private/Widgets/DataprepGraph/SDataprepActionBlock.h
+++ b/Engine/Plugins/Enterprise/DataprepEditor/Source/DataprepEditor/Private/Widgets/DataprepGraph/SDataprepActionBlock.h
@@ -35,8 +35,10 @@ public:
 	/** Return the block title */
 	TSharedRef<SWidget> GetBlockTitleWidget();
 
-	/** Temporary while both graph editors coexist */
-	bool bDragEnabled;
+#ifndef NO_BLUEPRINT
+	// Temporary boolean to indicate if the widget has been instantiated in a regular BP graph or the simplified one
+	bool bIsSimplifiedGraph;
+#endif
 
 protected:
 
diff --git a/Engine/Plugins/Enterprise/DataprepEditor/Source/DataprepEditor/Private/Widgets/DataprepGraph/SDataprepGraphActionStepNode.cpp b/Engine/Plugins/Enterprise/DataprepEditor/Source/DataprepEditor/Private/Widgets/DataprepGraph/SDataprepGraphActionStepNode.cpp
--- a/Engine/Plugins/Enterprise/DataprepEditor/Source/DataprepEditor/Private/Widgets/DataprepGraph/SDataprepGraphActionStepNode.cpp
+++ b/Engine/Plugins/Enterprise/DataprepEditor/Source/DataprepEditor/Private/Widgets/DataprepGraph/SDataprepGraphActionStepNode.cpp
@@ -105,7 +105,9 @@ void SDataprepGraphActionStepNode::UpdateGraphNode()
 
 			if(ActionStepBlockPtr.IsValid())
 			{
-				ActionStepBlockPtr->bDragEnabled = false;
+#ifndef NO_BLUEPRINT
+				ActionStepBlockPtr->bIsSimplifiedGraph = true;
+#endif
 				ActionBlock = ActionStepBlockPtr->AsShared();
 			}
 		}
