Integrating CL #11213699 and #11213849 from //UE4/Dev-Enterprise

diff --git a/Engine/Plugins/Enterprise/DataprepEditor/Source/DataprepEditorScriptingUtilities/Private/EditorDataprepAssetLibrary.cpp b/Engine/Plugins/Enterprise/DataprepEditor/Source/DataprepEditorScriptingUtilities/Private/EditorDataprepAssetLibrary.cpp
--- a/Engine/Plugins/Enterprise/DataprepEditor/Source/DataprepEditorScriptingUtilities/Private/EditorDataprepAssetLibrary.cpp
+++ b/Engine/Plugins/Enterprise/DataprepEditor/Source/DataprepEditorScriptingUtilities/Private/EditorDataprepAssetLibrary.cpp
@@ -91,7 +91,7 @@ bool UEditorDataprepAssetLibrary::ExecuteDataprep(UDataprepAssetInterface* Datap
 
 		switch (LogReportingMethod)
 		{
-		case EDataprepReportMethod::StandartLog:
+		case EDataprepReportMethod::StandardLog:
 		case EDataprepReportMethod::SameFeedbackAsEditor:
 			Logger = MakeShared<FDataprepCoreUtils::FDataprepLogger>();
 			break;
@@ -102,7 +102,7 @@ bool UEditorDataprepAssetLibrary::ExecuteDataprep(UDataprepAssetInterface* Datap
 
 		switch (ProgressReportingMethod)
 		{
-		case EDataprepReportMethod::StandartLog:
+		case EDataprepReportMethod::StandardLog:
 			Reporter = MakeShared<FDataprepCoreUtils::FDataprepProgressTextReporter>();
 			break;
 		case EDataprepReportMethod::SameFeedbackAsEditor:
diff --git a/Engine/Plugins/Enterprise/DataprepEditor/Source/DataprepEditorScriptingUtilities/Public/EditorDataprepAssetLibrary.h b/Engine/Plugins/Enterprise/DataprepEditor/Source/DataprepEditorScriptingUtilities/Public/EditorDataprepAssetLibrary.h
--- a/Engine/Plugins/Enterprise/DataprepEditor/Source/DataprepEditorScriptingUtilities/Public/EditorDataprepAssetLibrary.h
+++ b/Engine/Plugins/Enterprise/DataprepEditor/Source/DataprepEditorScriptingUtilities/Public/EditorDataprepAssetLibrary.h
@@ -19,7 +19,7 @@ UENUM(BlueprintType)
 enum class EDataprepReportMethod : uint8
 {
 	// Report the feedback into the output log only
-	StandartLog,
+	StandardLog,
 
 	// Report the feedback the same way that the dataprep asset editor does
 	SameFeedbackAsEditor,
