Integrating CL #11218726 from //UE4/Dev-Enterprise

Fixed renaming of action node

diff --git a/Engine/Plugins/Enterprise/DataprepEditor/Source/DataprepCore/Private/DataprepAsset.cpp b/Engine/Plugins/Enterprise/DataprepEditor/Source/DataprepCore/Private/DataprepAsset.cpp
--- a/Engine/Plugins/Enterprise/DataprepEditor/Source/DataprepCore/Private/DataprepAsset.cpp
+++ b/Engine/Plugins/Enterprise/DataprepEditor/Source/DataprepCore/Private/DataprepAsset.cpp
@@ -539,13 +539,14 @@ void UDataprepAsset::UpdateActions(bool bNotify)
 
 int32 UDataprepAsset::AddAction(const UDataprepActionAsset* InAction)
 {
-	if ( InAction )
+	UDataprepActionAsset* Action = InAction ? DuplicateObject<UDataprepActionAsset>( InAction, this) : NewObject<UDataprepActionAsset>( this, UDataprepActionAsset::StaticClass(), NAME_None, RF_Transactional );
+
+	if ( Action )
 	{
 		Modify();
 
-		UDataprepActionAsset* Action = DuplicateObject<UDataprepActionAsset>( InAction, this);
 		Action->SetFlags(EObjectFlags::RF_Transactional);
-		Action->SetLabel( InAction->GetLabel() );
+		Action->SetLabel( InAction ? InAction->GetLabel() : TEXT("New Action") );
 
 		ActionAssets.Add( Action );
 		OnActionChanged.Broadcast(Action, FDataprepAssetChangeType::ActionAdded);
diff --git a/Engine/Plugins/Enterprise/DataprepEditor/Source/DataprepCore/Public/DataprepActionAsset.h b/Engine/Plugins/Enterprise/DataprepEditor/Source/DataprepCore/Public/DataprepActionAsset.h
--- a/Engine/Plugins/Enterprise/DataprepEditor/Source/DataprepCore/Public/DataprepActionAsset.h
+++ b/Engine/Plugins/Enterprise/DataprepEditor/Source/DataprepCore/Public/DataprepActionAsset.h
@@ -330,10 +330,10 @@ public:
 
 	/** Getter and Setter on the UI text of the action */
 	const TCHAR* GetLabel() const { return *Label; }
-	void SetLabel( const TCHAR* InLabel ) { Label = InLabel ? InLabel : TEXT(""); }
+	void SetLabel( const TCHAR* InLabel ) { Modify(); Label = InLabel ? InLabel : TEXT(""); }
 
 	/**
-	 * Do the necessary notification so that the dataprep system can react properly to removal of this action
+	 * Do the necessary notification so that the Dataprep system can react properly to removal of this action
 	 */
 	void NotifyDataprepSystemsOfRemoval();
 
diff --git a/Engine/Plugins/Enterprise/DataprepEditor/Source/DataprepCore/Public/DataprepAsset.h b/Engine/Plugins/Enterprise/DataprepEditor/Source/DataprepCore/Public/DataprepAsset.h
--- a/Engine/Plugins/Enterprise/DataprepEditor/Source/DataprepCore/Public/DataprepAsset.h
+++ b/Engine/Plugins/Enterprise/DataprepEditor/Source/DataprepCore/Public/DataprepAsset.h
@@ -70,8 +70,9 @@ public:
 
 	/**
 	 * Add a copy of the action to the Dataprep asset
-	 * @param Action The action we want to duplicate in the Dataprep asset
+	 * @param Action The action we want to duplicate in the Dataprep asset. Parameter can be null
 	 * @return The index of the added action or index none if the action is invalid
+	 * @remark If action is nullptr, a new DataprepActionAsset is simply created
 	 */
 	int32 AddAction(const UDataprepActionAsset* Action);
 
diff --git a/Engine/Plugins/Enterprise/DataprepEditor/Source/DataprepEditor/Private/DataprepEditor.cpp b/Engine/Plugins/Enterprise/DataprepEditor/Source/DataprepEditor/Private/DataprepEditor.cpp
--- a/Engine/Plugins/Enterprise/DataprepEditor/Source/DataprepEditor/Private/DataprepEditor.cpp
+++ b/Engine/Plugins/Enterprise/DataprepEditor/Source/DataprepEditor/Private/DataprepEditor.cpp
@@ -431,6 +431,7 @@ void FDataprepEditor::InitDataprepEditor(const EToolkitMode::Type Mode, const TS
 	TempDir = FPaths::Combine( GetRootTemporaryDir(), FString::FromInt( FPlatformProcess::GetCurrentProcessId() ), SessionID);
 	IFileManager::Get().MakeDirectory(*TempDir);
 
+#ifndef NO_BLUEPRINT
 	// Temp code for the nodes development
 	if(Blueprint != nullptr)
 	{
@@ -439,11 +440,12 @@ void FDataprepEditor::InitDataprepEditor(const EToolkitMode::Type Mode, const TS
 
 		// Necessary step to regenerate blueprint generated class
 		// Note that this compilation will always succeed as Dataprep node does not have real body
-		{
-			FKismetEditorUtilities::CompileBlueprint( DataprepRecipeBPPtr.Get(), EBlueprintCompileOptions::None, nullptr );
-		}
+		//{
+		//	FKismetEditorUtilities::CompileBlueprint( DataprepRecipeBPPtr.Get(), EBlueprintCompileOptions::None, nullptr );
+		//}
 	}
 	// End temp code for the nodes development
+#endif
 
 	GEditor->RegisterForUndo(this);
 
diff --git a/Engine/Plugins/Enterprise/DataprepEditor/Source/DataprepEditor/Private/DataprepGraph/DataprepGraph.cpp b/Engine/Plugins/Enterprise/DataprepEditor/Source/DataprepEditor/Private/DataprepGraph/DataprepGraph.cpp
--- a/Engine/Plugins/Enterprise/DataprepEditor/Source/DataprepEditor/Private/DataprepGraph/DataprepGraph.cpp
+++ b/Engine/Plugins/Enterprise/DataprepEditor/Source/DataprepEditor/Private/DataprepGraph/DataprepGraph.cpp
@@ -54,12 +54,6 @@ void FDataprepEditor::CreateGraphEditor()
 			]
 		];
 
-		SGraphEditor::FGraphEditorEvents Events;
-		//Events.OnSelectionChanged = SGraphEditor::FOnSelectionChanged::CreateSP(this, &FDataprepEditor::OnPipelineEditorSelectionChanged);
-		//Events.OnCreateActionMenu = SGraphEditor::FOnCreateActionMenu::CreateSP(this, &FDataprepEditor::OnCreatePipelineActionMenu);
-		//Events.OnVerifyTextCommit = FOnNodeVerifyTextCommit::CreateSP(this, &FDataprepEditor::OnNodeVerifyTitleCommit);
-		//Events.OnTextCommitted = FOnNodeTextCommitted::CreateSP(this, &FDataprepEditor::OnNodeTitleCommitted);
-
 		FName UniqueGraphName = MakeUniqueObjectName( GetTransientPackage(), UWorld::StaticClass(), FName( *(LOCTEXT("DataprepGraph", "Graph").ToString()) ) );
 		DataprepGraph = TStrongObjectPtr<UDataprepGraph>( NewObject< UDataprepGraph >(GetTransientPackage(), UniqueGraphName) );
 		DataprepGraph->Schema = UDataprepGraphSchema::StaticClass();
@@ -69,8 +63,7 @@ void FDataprepEditor::CreateGraphEditor()
 		GraphEditor = SNew(SDataprepGraphEditor, DataprepAsset)
 			.Appearance(AppearanceInfo)
 			.TitleBar(TitleBarWidget)
-			.GraphToEdit(DataprepGraph.Get())
-			.GraphEvents(Events);
+			.GraphToEdit(DataprepGraph.Get());
 	}
 }
 
diff --git a/Engine/Plugins/Enterprise/DataprepEditor/Source/DataprepEditor/Private/Widgets/DataprepGraph/SDataprepGraphActionNode.cpp b/Engine/Plugins/Enterprise/DataprepEditor/Source/DataprepEditor/Private/Widgets/DataprepGraph/SDataprepGraphActionNode.cpp
--- a/Engine/Plugins/Enterprise/DataprepEditor/Source/DataprepEditor/Private/Widgets/DataprepGraph/SDataprepGraphActionNode.cpp
+++ b/Engine/Plugins/Enterprise/DataprepEditor/Source/DataprepEditor/Private/Widgets/DataprepGraph/SDataprepGraphActionNode.cpp
@@ -314,6 +314,28 @@ FReply SDataprepGraphActionNode::OnMouseButtonDown(const FGeometry& MyGeometry,
 		return FReply::Handled().DetectDrag( AsShared(), EKeys::LeftMouseButton );
 	}
 
+	// Take ownership of the mouse if right mouse button clicked to display contextual menu
+	if ( MouseEvent.GetEffectingButton() == EKeys::RightMouseButton )
+	{
+		return FReply::Handled();
+	}
+
+	return FReply::Unhandled();
+}
+
+FReply SDataprepGraphActionNode::OnMouseButtonUp(const FGeometry & MyGeometry, const FPointerEvent & MouseEvent)
+{
+	if ( MouseEvent.GetEffectingButton() == EKeys::RightMouseButton )
+	{
+		ensure(OwnerGraphPanelPtr.IsValid());
+
+		const FVector2D Position = MouseEvent.GetScreenSpacePosition();
+		OwnerGraphPanelPtr.Pin()->SummonContextMenu(Position, Position, GraphNode, nullptr, TArray<UEdGraphPin*>());
+
+		// Release mouse capture
+		return FReply::Handled().ReleaseMouseCapture();
+	}
+
 	return FReply::Unhandled();
 }
 
@@ -330,6 +352,7 @@ void SDataprepGraphActionNode::SetOwner(const TSharedRef<SGraphPanel>& OwnerPane
 	if(!OwnerGraphPanelPtr.IsValid())
 	{
 		SGraphNode::SetOwner(OwnerPanel);
+		OwnerPanel->AttachGraphEvents(SharedThis(this));
 
 		OwnerPanel->AddGraphNode(SharedThis(ProxyNodePtr.Get()));
 
diff --git a/Engine/Plugins/Enterprise/DataprepEditor/Source/DataprepEditor/Private/Widgets/DataprepGraph/SDataprepGraphActionNode.h b/Engine/Plugins/Enterprise/DataprepEditor/Source/DataprepEditor/Private/Widgets/DataprepGraph/SDataprepGraphActionNode.h
--- a/Engine/Plugins/Enterprise/DataprepEditor/Source/DataprepEditor/Private/Widgets/DataprepGraph/SDataprepGraphActionNode.h
+++ b/Engine/Plugins/Enterprise/DataprepEditor/Source/DataprepEditor/Private/Widgets/DataprepGraph/SDataprepGraphActionNode.h
@@ -29,6 +29,7 @@ public:
 
 	// SWidget interface
 	virtual FReply OnMouseButtonDown(const FGeometry& MyGeometry, const FPointerEvent& MouseEvent) override;
+	virtual FReply OnMouseButtonUp(const FGeometry& MyGeometry, const FPointerEvent& MouseEvent) override;
 	virtual FReply OnDragDetected(const FGeometry& MyGeometry, const FPointerEvent& MouseEvent) override;
 	virtual FCursorReply OnCursorQuery(const FGeometry& MyGeometry, const FPointerEvent& CursorEvent) const override;
 	virtual void SetOwner(const TSharedRef<SGraphPanel>& OwnerPanel) override;
diff --git a/Engine/Plugins/Enterprise/DataprepEditor/Source/DataprepEditor/Private/Widgets/DataprepGraph/SDataprepGraphEditor.cpp b/Engine/Plugins/Enterprise/DataprepEditor/Source/DataprepEditor/Private/Widgets/DataprepGraph/SDataprepGraphEditor.cpp
--- a/Engine/Plugins/Enterprise/DataprepEditor/Source/DataprepEditor/Private/Widgets/DataprepGraph/SDataprepGraphEditor.cpp
+++ b/Engine/Plugins/Enterprise/DataprepEditor/Source/DataprepEditor/Private/Widgets/DataprepGraph/SDataprepGraphEditor.cpp
@@ -8,22 +8,27 @@
 #include "DataprepGraph/DataprepGraph.h"
 #include "DataprepGraph/DataprepGraphSchema.h"
 #include "DataprepGraph/DataprepGraphActionNode.h"
+#include "SchemaActions/DataprepAllMenuActionCollector.h"
 #include "SchemaActions/DataprepDragDropOp.h"
+#include "SchemaActions/IDataprepMenuActionCollector.h"
 #include "Widgets/DataprepGraph/SDataprepGraphActionNode.h"
 #include "Widgets/DataprepGraph/SDataprepGraphActionStepNode.h"
 #include "Widgets/DataprepGraph/SDataprepGraphTrackNode.h"
 #include "Widgets/DataprepWidgets.h"
+#include "Widgets/SDataprepActionMenu.h"
 
 #include "EdGraph/EdGraphSchema.h"
 #include "EdGraphNode_Comment.h"
 #include "EdGraphUtilities.h"
 #include "Framework/Commands/GenericCommands.h"
 #include "Framework/Commands/UICommandList.h"
+#include "Framework/MultiBox/MultiBoxBuilder.h"
 #include "HAL/PlatformApplicationMisc.h"
 #include "Kismet2/Kismet2NameValidators.h"
 #include "ScopedTransaction.h"
 #include "SGraphPanel.h"
 #include "Widgets/Colors/SColorBlock.h"
+#include "Widgets/Input/SEditableTextBox.h"
 #include "Widgets/Layout/SBox.h"
 #include "Widgets/SBoxPanel.h"
 #include "Widgets/SOverlay.h"
@@ -77,13 +82,19 @@ void SDataprepGraphEditor::Construct(const FArguments& InArgs, UDataprepAsset* I
 	check(InDataprepAsset);
 	DataprepAssetPtr = InDataprepAsset;
 
+	SGraphEditor::FGraphEditorEvents Events;
+	Events.OnCreateNodeOrPinMenu = SGraphEditor::FOnCreateNodeOrPinMenu::CreateSP(this, &SDataprepGraphEditor::OnCreateNodeOrPinMenu);
+	Events.OnCreateActionMenu = SGraphEditor::FOnCreateActionMenu::CreateSP(this, &SDataprepGraphEditor::OnCreateActionMenu);
+	Events.OnVerifyTextCommit = FOnNodeVerifyTextCommit::CreateSP(this, &SDataprepGraphEditor::OnNodeVerifyTitleCommit);
+	Events.OnTextCommitted = FOnNodeTextCommitted::CreateSP(this, &SDataprepGraphEditor::OnNodeTitleCommitted);
+
 	BuildCommandList();
 
 	SGraphEditor::FArguments Arguments;
 	Arguments._AdditionalCommands = GraphEditorCommands;
 	Arguments._TitleBar = InArgs._TitleBar;
 	Arguments._GraphToEdit = InArgs._GraphToEdit;
-	Arguments._GraphEvents = InArgs._GraphEvents;
+	Arguments._GraphEvents = Events;
 
 	SGraphEditor::Construct( Arguments );
 
@@ -399,10 +410,6 @@ void SDataprepGraphEditor::BuildCommandList()
 			FExecuteAction::CreateSP(this, &SDataprepGraphEditor::DuplicateNodes),
 			FCanExecuteAction::CreateSP(this, &SDataprepGraphEditor::CanDuplicateNodes)
 		);
-
-		GraphEditorCommands->MapAction(FGraphEditorCommands::Get().CreateComment,
-			FExecuteAction::CreateSP(this, &SDataprepGraphEditor::OnCreateComment)
-		);
 	}
 }
 
@@ -705,31 +712,6 @@ void SDataprepGraphEditor::DeleteSelectedDuplicatableNodes()
 	}
 }
 
-void SDataprepGraphEditor::OnCreateComment()
-{
-	if (UEdGraph* Graph = GetCurrentGraph())
-	{
-		if (const UEdGraphSchema* Schema = Graph->GetSchema())
-		{
-			// Add menu item for creating comment boxes
-			UEdGraphNode_Comment* CommentTemplate = NewObject<UEdGraphNode_Comment>();
-
-			FVector2D SpawnLocation = GetPasteLocation();
-
-			FSlateRect Bounds;
-			if ( GetBoundsForSelectedNodes(Bounds, 50.0f) )
-			{
-				CommentTemplate->SetBounds(Bounds);
-				SpawnLocation.X = CommentTemplate->NodePosX;
-				SpawnLocation.Y = CommentTemplate->NodePosY;
-			}
-
-			UEdGraphNode* NewNode = FEdGraphSchemaAction_NewNode::SpawnNodeFromTemplate<UEdGraphNode_Comment>(Graph, CommentTemplate, SpawnLocation, /** bSelectNewNode */ true);
-		}
-	}
-}
-
-
 bool SDataprepGraphEditor::OnNodeVerifyTitleCommit(const FText& NewText, UEdGraphNode* NodeBeingChanged, FText& OutErrorMessage)
 {
 	bool bValid(false);
@@ -749,14 +731,6 @@ bool SDataprepGraphEditor::OnNodeVerifyTitleCommit(const FText& NewText, UEdGrap
 		{
 			bValid = true;
 		}
-		//else if (PipelineView.IsValid())
-		//{
-		//	EValidatorResult Valid = NameEntryValidator->IsValid(NewText.ToString(), false);
-
-		//	NodeBeingChanged->bHasCompilerMessage = true;
-		//	NodeBeingChanged->ErrorMsg = NameEntryValidator->GetErrorString(NewText.ToString(), Valid);
-		//	NodeBeingChanged->ErrorType = EMessageSeverity::Error;
-		//}
 	}
 
 	return bValid;
@@ -772,4 +746,46 @@ void SDataprepGraphEditor::OnNodeTitleCommitted(const FText& NewText, ETextCommi
 	}
 }
 
+FActionMenuContent SDataprepGraphEditor::OnCreateActionMenu(UEdGraph* InGraph, const FVector2D& InNodePosition, const TArray<UEdGraphPin*>& InDraggedPins, bool bAutoExpand, SGraphEditor::FActionMenuClosed InOnMenuClosed)
+{
+	// bAutoExpand is voluntary ignored for now
+	TUniquePtr<IDataprepMenuActionCollector> ActionCollector = MakeUnique<FDataprepAllMenuActionCollector>();
+
+	TSharedRef<SDataprepActionMenu> ActionMenu =
+		SNew (SDataprepActionMenu, MoveTemp(ActionCollector))
+		.TransactionText(LOCTEXT("AddingANewActionNode","Add a Action Node"))
+		.GraphObj(InGraph)
+		.NewNodePosition(InNodePosition)
+		.DraggedFromPins(InDraggedPins)
+		.OnClosedCallback(InOnMenuClosed);
+
+	return FActionMenuContent( ActionMenu, ActionMenu->GetFilterTextBox() );
+}
+
+FActionMenuContent SDataprepGraphEditor::OnCreateNodeOrPinMenu(UEdGraph* CurrentGraph, const UEdGraphNode* InGraphNode, const UEdGraphPin* InGraphPin, FMenuBuilder* MenuBuilder, bool bIsDebugging)
+{
+	if(CurrentGraph != GetCurrentGraph())
+	{
+		return FActionMenuContent();
+	}
+
+	// Open contextual menu for action node
+	if(const UDataprepGraphActionNode* ActionNode = Cast<const UDataprepGraphActionNode>(InGraphNode))
+	{
+		MenuBuilder->BeginSection( FName( TEXT("CommonSection") ), LOCTEXT("CommonSection", "Common") );
+		{
+			MenuBuilder->AddMenuEntry(FGenericCommands::Get().Copy);
+			MenuBuilder->AddMenuEntry(FGenericCommands::Get().Cut);
+			MenuBuilder->AddMenuEntry(FGenericCommands::Get().Duplicate);
+			MenuBuilder->AddMenuEntry(FGenericCommands::Get().Delete);
+		}
+		MenuBuilder->EndSection();
+
+		return FActionMenuContent(MenuBuilder->MakeWidget());
+	}
+
+	// Create contextual for graph panel when on top of track node too
+	return OnCreateActionMenu(CurrentGraph, GetPasteLocation(), TArray<UEdGraphPin*>(), true, SGraphEditor::FActionMenuClosed());
+}
+
 #undef LOCTEXT_NAMESPACE
\ No newline at end of file
diff --git a/Engine/Plugins/Enterprise/DataprepEditor/Source/DataprepEditor/Private/Widgets/DataprepGraph/SDataprepGraphEditor.h b/Engine/Plugins/Enterprise/DataprepEditor/Source/DataprepEditor/Private/Widgets/DataprepGraph/SDataprepGraphEditor.h
--- a/Engine/Plugins/Enterprise/DataprepEditor/Source/DataprepEditor/Private/Widgets/DataprepGraph/SDataprepGraphEditor.h
+++ b/Engine/Plugins/Enterprise/DataprepEditor/Source/DataprepEditor/Private/Widgets/DataprepGraph/SDataprepGraphEditor.h
@@ -93,13 +93,19 @@ public:
 	void DuplicateNodes();
 	bool CanDuplicateNodes() const;
 
-	void OnCreateComment();
-
 	void DeleteSelectedDuplicatableNodes();
+	/** End of set of methods necessary for copy/paste of action nodes*/
 
+	/** Set of callbacks to rename an action node */
 	bool OnNodeVerifyTitleCommit(const FText& NewText, UEdGraphNode* NodeBeingChanged, FText& OutErrorMessage);
 	void OnNodeTitleCommitted(const FText& NewText, ETextCommit::Type CommitInfo, UEdGraphNode* NodeBeingChanged);
 
+	/** Callback to create contextual menu on graph panel */
+	FActionMenuContent OnCreateActionMenu(UEdGraph* InGraph, const FVector2D& InNodePosition, const TArray<UEdGraphPin*>& InDraggedPins, bool bAutoExpand, SGraphEditor::FActionMenuClosed InOnMenuClosed);
+
+	/** Callback to create contextual menu for graph nodes in graph panel */
+	FActionMenuContent OnCreateNodeOrPinMenu(UEdGraph* CurrentGraph, const UEdGraphNode* InGraphNode, const UEdGraphPin* InGraphPin, FMenuBuilder* MenuBuilder, bool bIsDebugging);
+
 private:
 	/** Recompute the layout of the displayed graph after a pan, resize and/or zoom */
 	void UpdateLayout( const FVector2D& LocalSize, const FVector2D& Location, float ZoomAmount );
diff --git a/Engine/Plugins/Enterprise/DataprepEditor/Source/DataprepEditor/Private/Widgets/SDataprepActionMenu.cpp b/Engine/Plugins/Enterprise/DataprepEditor/Source/DataprepEditor/Private/Widgets/SDataprepActionMenu.cpp
--- a/Engine/Plugins/Enterprise/DataprepEditor/Source/DataprepEditor/Private/Widgets/SDataprepActionMenu.cpp
+++ b/Engine/Plugins/Enterprise/DataprepEditor/Source/DataprepEditor/Private/Widgets/SDataprepActionMenu.cpp
@@ -4,6 +4,8 @@
 
 #include "BlueprintNodes/K2Node_DataprepAction.h"
 #include "DataprepActionAsset.h"
+#include "DataprepAsset.h"
+#include "DataprepGraph/DataprepGraph.h"
 #include "SGraphActionMenu.h"
 #include "SchemaActions/DataprepSchemaAction.h"
 #include "SchemaActions/DataprepSchemaActionUtils.h"
@@ -69,38 +71,52 @@ void SDataprepActionMenu::CollectActions(FGraphActionListBuilderBase& OutActions
 
 void SDataprepActionMenu::OnActionSelected(const TArray< TSharedPtr<FEdGraphSchemaAction> >& SelectedActions, ESelectInfo::Type InSelectionType)
 {
+#ifndef NO_BLUEPRINT
 	/**
 	 * Create a node if needed and execute the action
 	 * return true if the action was properly executed
 	 */
-	auto ExcuteAction = [this](FDataprepSchemaAction& DataprepAction) -> bool
+	auto ExecuteActionBP = [this](FDataprepSchemaAction& DataprepAction) -> bool
+	{
+		UK2Node_DataprepAction* ActionNode = nullptr;
+
+		if ( ShouldCreateNewNode() )
 		{
-			UK2Node_DataprepAction* ActionNode = nullptr;
+			ActionNode = DataprepSchemaActionUtils::SpawnEdGraphNode<UK2Node_DataprepAction>( *GraphObj, NewNodePosition );
+			ActionNode->CreateDataprepActionAsset();
+			Context.DataprepActionPtr = MakeWeakObjectPtr<UDataprepActionAsset>( ActionNode->GetDataprepAction() );
 
-			if ( ShouldCreateNewNode() )
+			if ( DraggedFromPins.Num() > 0 )
 			{
-				ActionNode = DataprepSchemaActionUtils::SpawnEdGraphNode<UK2Node_DataprepAction>( *GraphObj, NewNodePosition );
-				ActionNode->CreateDataprepActionAsset();
-				Context.DataprepActionPtr = MakeWeakObjectPtr<UDataprepActionAsset>( ActionNode->GetDataprepAction() );
-
-				if ( DraggedFromPins.Num() > 0 )
-				{
-					ActionNode->AutowireNewNode( DraggedFromPins[0] );
-				}
+				ActionNode->AutowireNewNode( DraggedFromPins[0] );
 			}
+		}
 
-			if ( !DataprepAction.ExecuteAction( Context ) )
+		if ( !DataprepAction.ExecuteAction( Context ) )
+		{
+			if ( ActionNode )
 			{
-				if ( ActionNode )
-				{
-					// Remove the created node
-					GraphObj->RemoveNode( ActionNode );
-				}
-				return false;
+				// Remove the created node
+				GraphObj->RemoveNode( ActionNode );
 			}
+			return false;
+		}
+
+		return true;
+	};
+#endif
+	auto ExecuteAction = [this](FDataprepSchemaAction& DataprepAction, UDataprepGraph* DataprepGraph) -> bool
+	{
+		UDataprepAsset* DataprepAsset = DataprepGraph ? DataprepGraph->GetDataprepAsset() : nullptr;
+
+		if ( DataprepAsset )
+		{
+			int32 Index = DataprepAsset->AddAction(nullptr);
+			Context.DataprepActionPtr = MakeWeakObjectPtr<UDataprepActionAsset>( DataprepAsset->GetAction(Index) );
+		}
 
-			return true;
-		};
+		return DataprepAction.ExecuteAction( Context ) ? true : false;
+	};
 
 	for ( const TSharedPtr<FEdGraphSchemaAction>& Action : SelectedActions )
 	{
@@ -111,14 +127,32 @@ void SDataprepActionMenu::OnActionSelected(const TArray< TSharedPtr<FEdGraphSche
 			{
 				FScopedTransaction Transaction( TransactionTextGetter.Get() );
 
-				if ( !ExcuteAction( *DataprepAction ) )
+				if(UDataprepGraph* DataprepGraph = Cast<UDataprepGraph>(GraphObj))
+				{
+					if ( !ExecuteAction( *DataprepAction, DataprepGraph ) )
+					{
+						Transaction.Cancel();
+					}
+				}
+#ifndef NO_BLUEPRINT
+				else if ( !ExecuteActionBP( *DataprepAction ) )
 				{
 					Transaction.Cancel();
 				}
+#endif
 			}
 			else
 			{
-				ExcuteAction( *DataprepAction );
+				if(UDataprepGraph* DataprepGraph = Cast<UDataprepGraph>(GraphObj))
+				{
+					ExecuteAction( *DataprepAction, DataprepGraph );
+				}
+#ifndef NO_BLUEPRINT
+				else
+				{
+					ExecuteActionBP( *DataprepAction );
+				}
+#endif
 			}
 		}
 	}
