Integrating CL #11244459 from //UE4/Dev-Enterprise

Dataprep: Implement Reference Filter Transform

diff --git a/Engine/Plugins/Enterprise/DataprepEditor/Source/DataprepCore/Private/DataprepActionAsset.cpp b/Engine/Plugins/Enterprise/DataprepEditor/Source/DataprepCore/Private/DataprepActionAsset.cpp
--- a/Engine/Plugins/Enterprise/DataprepEditor/Source/DataprepCore/Private/DataprepActionAsset.cpp
+++ b/Engine/Plugins/Enterprise/DataprepEditor/Source/DataprepCore/Private/DataprepActionAsset.cpp
@@ -14,6 +14,7 @@
 #include "Parameterization/DataprepParameterization.h"
 #include "SelectionSystem/DataprepFetcher.h"
 #include "SelectionSystem/DataprepFilter.h"
+#include "SelectionSystem/DataprepSelectionTransform.h"
 
 // Engine include
 #include "ActorEditorUtils.h"
@@ -124,6 +125,13 @@ void UDataprepActionAsset::Execute(const TArray<UObject*>& InObjects)
 				UDataprepFilter* Filter = static_cast<UDataprepFilter*>( StepObject );
 				OperationContext->Context->Objects = Filter->FilterObjects( OperationContext->Context->Objects );
 			}
+			else if ( StepType == UDataprepSelectionTransform::StaticClass() )
+			{
+				UDataprepSelectionTransform* SelectionTransform = static_cast<UDataprepSelectionTransform*>(StepObject);
+				TArray<UObject*> TransformResult;
+				SelectionTransform->Execute(OperationContext->Context->Objects, TransformResult);
+				OperationContext->Context->Objects = TransformResult;
+			}
 		}
 	}
 
@@ -161,6 +169,17 @@ int32 UDataprepActionAsset::AddStep(TSubclassOf<UDataprepParameterizableObject>
 			OnStepsOrderChanged.Broadcast();
 			return Steps.Num() - 1;
 		}
+		if ( ValidRootClass == UDataprepSelectionTransform::StaticClass() )
+		{
+			Modify();
+			UDataprepActionStep* ActionStep = NewObject< UDataprepActionStep >(this, UDataprepActionStep::StaticClass(), NAME_None, RF_Transactional);
+			ActionStep->StepObject = NewObject< UDataprepSelectionTransform >(ActionStep, StepType.Get(), NAME_None, RF_Transactional);
+			ActionStep->PathOfStepObjectClass = ActionStep->StepObject->GetClass();
+			ActionStep->bIsEnabled = true;
+			Steps.Add(ActionStep);
+			OnStepsOrderChanged.Broadcast();
+			return Steps.Num() - 1;
+		}
 		else if ( ValidRootClass == UDataprepFetcher::StaticClass() )
 		{
 			if ( UClass* FilterClass = UDataprepFilter::GetFilterTypeForFetcherType( StepType.Get() ) )
@@ -583,6 +602,13 @@ void UDataprepActionAsset::ExecuteAction(const TSharedPtr<FDataprepActionContext
 			UDataprepFilter* Filter = static_cast<UDataprepFilter*>( StepObject );
 			SelectedObjects = Filter->FilterObjects( SelectedObjects );
 		}
+		else if ( StepType == UDataprepSelectionTransform::StaticClass() )
+		{
+			UDataprepSelectionTransform* SelectionTransform = static_cast<UDataprepSelectionTransform*>(StepObject);
+			TArray<UObject*> TransformResult;
+			SelectionTransform->Execute(SelectedObjects, TransformResult);
+			SelectedObjects = TransformResult;
+		}
 	};
 
 	if(SpecificStep && bSpecificStepOnly == true)
diff --git a/Engine/Plugins/Enterprise/DataprepEditor/Source/DataprepCore/Private/DataprepCoreUtils.cpp b/Engine/Plugins/Enterprise/DataprepEditor/Source/DataprepCore/Private/DataprepCoreUtils.cpp
--- a/Engine/Plugins/Enterprise/DataprepEditor/Source/DataprepCore/Private/DataprepCoreUtils.cpp
+++ b/Engine/Plugins/Enterprise/DataprepEditor/Source/DataprepCore/Private/DataprepCoreUtils.cpp
@@ -18,6 +18,7 @@
 #include "IDataprepProgressReporter.h"
 #include "SelectionSystem/DataprepFetcher.h"
 #include "SelectionSystem/DataprepFilter.h"
+#include "SelectionSystem/DataprepSelectionTransform.h"
 
 #include "Engine/StaticMesh.h"
 #include "Engine/World.h"
@@ -369,6 +370,7 @@ bool FDataprepCoreUtils::IsClassValidForStepCreation(const TSubclassOf<UDataprep
 	UClass* DataprepTopLevelClass = UDataprepParameterizableObject::StaticClass();
 	UClass* DataprepOperationClass = UDataprepOperation::StaticClass();
 	UClass* DataprepFetcherClass = UDataprepFetcher::StaticClass();
+	UClass* DataprepTransformClass = UDataprepSelectionTransform::StaticClass();
 
 	while ( Class )
 	{
@@ -396,6 +398,12 @@ bool FDataprepCoreUtils::IsClassValidForStepCreation(const TSubclassOf<UDataprep
 			return true;
 		}
 
+		if ( Class == DataprepTransformClass )
+		{
+			OutValidRootClass = DataprepTransformClass;
+			return true;
+		}
+
 		Class = Class->GetSuperClass();
 	}
 
@@ -408,11 +416,13 @@ UClass* FDataprepCoreUtils::GetTypeOfActionStep(const UDataprepParameterizableOb
 
 	const UClass* DataprepFilterClass = UDataprepFilter::StaticClass();
 	const UClass* DataprepOperationClass = UDataprepOperation::StaticClass();
+	const UClass* DataprepTransformClass = UDataprepSelectionTransform::StaticClass();
 
 	while ( CurrentClass )
 	{
 		if ( CurrentClass == DataprepFilterClass
 			|| CurrentClass == DataprepOperationClass
+			|| CurrentClass == DataprepTransformClass
 			)
 		{
 			break;
diff --git a/Engine/Plugins/Enterprise/DataprepEditor/Source/DataprepEditor/Private/SchemaActions/DataprepAllMenuActionCollector.cpp b/Engine/Plugins/Enterprise/DataprepEditor/Source/DataprepEditor/Private/SchemaActions/DataprepAllMenuActionCollector.cpp
--- a/Engine/Plugins/Enterprise/DataprepEditor/Source/DataprepEditor/Private/SchemaActions/DataprepAllMenuActionCollector.cpp
+++ b/Engine/Plugins/Enterprise/DataprepEditor/Source/DataprepEditor/Private/SchemaActions/DataprepAllMenuActionCollector.cpp
@@ -4,6 +4,7 @@
 
 #include "SchemaActions/DataprepFilterMenuActionCollector.h"
 #include "SchemaActions/DataprepOperationMenuActionCollector.h"
+#include "SchemaActions/DataprepSelectionTransformMenuActionCollector.h"
 #include "SchemaActions/DataprepSchemaAction.h"
 
 namespace FDataprepAllMenuActionCollectorUtils
@@ -27,6 +28,11 @@ TArray<TSharedPtr<FDataprepSchemaAction>> FDataprepAllMenuActionCollector::Colle
 	TArray< TSharedPtr< FDataprepSchemaAction > > Actions = FilterCollector.CollectActions();
 	FDataprepAllMenuActionCollectorUtils::AddRootCategoryToActions( Actions, FDataprepFilterMenuActionCollector::FilterCategory );
 
+	FDataprepSelectionTransformMenuActionCollector SelectionTransformCollector;
+	TArray< TSharedPtr< FDataprepSchemaAction > > TransformActions = SelectionTransformCollector.CollectActions();
+	FDataprepAllMenuActionCollectorUtils::AddRootCategoryToActions(TransformActions, FDataprepSelectionTransformMenuActionCollector::FilterCategory);
+	Actions.Append( MoveTemp( TransformActions ) );
+
 	FDataprepOperationMenuActionCollector OperationCollector;
 	TArray< TSharedPtr< FDataprepSchemaAction > > OperationActions = OperationCollector.CollectActions();
 	FDataprepAllMenuActionCollectorUtils::AddRootCategoryToActions( OperationActions, FDataprepOperationMenuActionCollector::OperationCategory );
diff --git a/Engine/Plugins/Enterprise/DataprepEditor/Source/DataprepEditor/Private/SchemaActions/DataprepSelectionTransformMenuActionCollector.cpp b/Engine/Plugins/Enterprise/DataprepEditor/Source/DataprepEditor/Private/SchemaActions/DataprepSelectionTransformMenuActionCollector.cpp
new file mode 100644
--- /dev/null
+++ b/Engine/Plugins/Enterprise/DataprepEditor/Source/DataprepEditor/Private/SchemaActions/DataprepSelectionTransformMenuActionCollector.cpp
@@ -0,0 +1,64 @@
+// Copyright Epic Games, Inc. All Rights Reserved.
+
+#include "SchemaActions/DataprepSelectionTransformMenuActionCollector.h"
+
+// Dataprep Includes
+#include "DataprepActionAsset.h"
+#include "DataprepEditorLogCategory.h"
+#include "DataprepMenuActionCollectorUtils.h"
+#include "SelectionSystem/DataprepSelectionTransform.h"
+#include "SchemaActions/DataprepSchemaAction.h"
+
+// Engine includes
+#include "UObject/Class.h"
+#include "Templates/SubclassOf.h"
+
+const  FText FDataprepSelectionTransformMenuActionCollector::FilterCategory = NSLOCTEXT("DataprepSelectionTransforms", "Selection Transforms Category", "Transform by");
+
+TArray<TSharedPtr<FDataprepSchemaAction>> FDataprepSelectionTransformMenuActionCollector::CollectActions()
+{
+	const double  Start = FPlatformTime::Seconds();
+
+	TArray< TSharedPtr< FDataprepSchemaAction > > Actions = DataprepMenuActionCollectorUtils::GatherMenuActionForDataprepClass(*UDataprepSelectionTransform::StaticClass()
+																															   , DataprepMenuActionCollectorUtils::FOnCreateMenuAction::CreateRaw(this, &FDataprepSelectionTransformMenuActionCollector::CreateMenuActionFromClass)
+	);
+
+	UE_LOG(LogDataprepEditor, Log, TEXT("The discovery of the selection transforms and the creation of the menu actions took %f seconds."), (FPlatformTime::Seconds() - Start));
+
+	return Actions;
+}
+
+bool FDataprepSelectionTransformMenuActionCollector::ShouldAutoExpand()
+{
+	return false;
+}
+
+TSharedPtr<FDataprepSchemaAction> FDataprepSelectionTransformMenuActionCollector::CreateMenuActionFromClass(UClass& Class)
+{
+	check(Class.IsChildOf<UDataprepSelectionTransform>());
+
+	const UDataprepSelectionTransform* SelectionTransform = static_cast<UDataprepSelectionTransform*>(Class.GetDefaultObject());
+	if (SelectionTransform)
+	{
+		FDataprepSchemaAction::FOnExecuteAction OnExcuteMenuAction;
+		OnExcuteMenuAction.BindLambda([Class = SelectionTransform->GetClass()](const FDataprepSchemaActionContext& InContext)
+		{
+			UDataprepActionAsset* Action = InContext.DataprepActionPtr.Get();
+			if (Action)
+			{
+				int32 NewOperationIndex = Action->AddStep(Class);
+				if (InContext.StepIndex != INDEX_NONE && InContext.StepIndex != NewOperationIndex)
+				{
+					Action->MoveStep(NewOperationIndex, InContext.StepIndex);
+				}
+			}
+		});
+
+		return MakeShared< FDataprepSchemaAction >(SelectionTransform->GetCategory()
+												   , SelectionTransform->GetDisplayTransformName(), SelectionTransform->GetTooltip()
+												   , 0, SelectionTransform->GetAdditionalKeyword(), OnExcuteMenuAction
+												   );
+	}
+
+	return {};
+}
diff --git a/Engine/Plugins/Enterprise/DataprepEditor/Source/DataprepEditor/Private/SchemaActions/DataprepSelectionTransformMenuActionCollector.h b/Engine/Plugins/Enterprise/DataprepEditor/Source/DataprepEditor/Private/SchemaActions/DataprepSelectionTransformMenuActionCollector.h
new file mode 100644
--- /dev/null
+++ b/Engine/Plugins/Enterprise/DataprepEditor/Source/DataprepEditor/Private/SchemaActions/DataprepSelectionTransformMenuActionCollector.h
@@ -0,0 +1,26 @@
+// Copyright Epic Games, Inc. All Rights Reserved.
+
+#pragma once
+
+#include "SchemaActions/IDataprepMenuActionCollector.h"
+
+#include "CoreMinimal.h"
+
+struct FDataprepSchemaAction;
+
+/**
+ * Help collecting all the dataprep selection transform menu action
+ */
+class FDataprepSelectionTransformMenuActionCollector : public IDataprepMenuActionCollector
+{
+public:
+	virtual TArray<TSharedPtr<FDataprepSchemaAction>> CollectActions() override;
+	virtual bool ShouldAutoExpand() override;
+
+	int32 GroupingPriority = 0;
+
+	static const FText FilterCategory;
+
+private:
+	TSharedPtr<FDataprepSchemaAction> CreateMenuActionFromClass(UClass& Class);
+};
\ No newline at end of file
diff --git a/Engine/Plugins/Enterprise/DataprepEditor/Source/DataprepEditor/Private/Widgets/DataprepGraph/SDataprepActionSteps.cpp b/Engine/Plugins/Enterprise/DataprepEditor/Source/DataprepEditor/Private/Widgets/DataprepGraph/SDataprepActionSteps.cpp
--- a/Engine/Plugins/Enterprise/DataprepEditor/Source/DataprepEditor/Private/Widgets/DataprepGraph/SDataprepActionSteps.cpp
+++ b/Engine/Plugins/Enterprise/DataprepEditor/Source/DataprepEditor/Private/Widgets/DataprepGraph/SDataprepActionSteps.cpp
@@ -10,7 +10,9 @@
 #include "SchemaActions/DataprepDragDropOp.h"
 #include "SchemaActions/DataprepSchemaAction.h"
 #include "SelectionSystem/DataprepFilter.h"
+#include "SelectionSystem/DataprepSelectionTransform.h"
 #include "Widgets/DataprepGraph/SDataprepFilter.h"
+#include "Widgets/DataprepGraph/SDataprepSelectionTransform.h"
 #include "Widgets/DataprepGraph/SDataprepOperation.h"
 #include "Widgets/SDataprepActionMenu.h"
 
@@ -37,6 +39,11 @@ namespace DataprepActionStepsWidgetUtils
 				UDataprepOperation* Operation = static_cast<UDataprepOperation*>( StepObject );
 				return SNew( SDataprepOperation, Operation, InStepData );
 			}
+			else if ( StepType == UDataprepSelectionTransform::StaticClass() )
+			{
+				UDataprepSelectionTransform* SelectionTransform = static_cast<UDataprepSelectionTransform*>( StepObject );
+				return SNew( SDataprepSelectionTransform, SelectionTransform, InStepData );
+			}
 			else if ( StepType == UDataprepFilter::StaticClass() )
 			{
 				UDataprepFilter* Filter = static_cast<UDataprepFilter*>( StepObject );
diff --git a/Engine/Plugins/Enterprise/DataprepEditor/Source/DataprepEditor/Private/Widgets/DataprepGraph/SDataprepGraphActionStepNode.cpp b/Engine/Plugins/Enterprise/DataprepEditor/Source/DataprepEditor/Private/Widgets/DataprepGraph/SDataprepGraphActionStepNode.cpp
--- a/Engine/Plugins/Enterprise/DataprepEditor/Source/DataprepEditor/Private/Widgets/DataprepGraph/SDataprepGraphActionStepNode.cpp
+++ b/Engine/Plugins/Enterprise/DataprepEditor/Source/DataprepEditor/Private/Widgets/DataprepGraph/SDataprepGraphActionStepNode.cpp
@@ -12,9 +12,11 @@
 #include "SchemaActions/DataprepDragDropOp.h"
 #include "SchemaActions/DataprepSchemaAction.h"
 #include "SelectionSystem/DataprepFilter.h"
+#include "SelectionSystem/DataprepSelectionTransform.h"
 #include "Widgets/DataprepGraph/SDataprepActionSteps.h"
 #include "Widgets/DataprepGraph/SDataprepFilter.h"
 #include "Widgets/DataprepGraph/SDataprepGraphActionNode.h"
+#include "Widgets/DataprepGraph/SDataprepSelectionTransform.h"
 #include "Widgets/DataprepGraph/SDataprepGraphTrackNode.h"
 #include "Widgets/DataprepGraph/SDataprepOperation.h"
 
@@ -102,6 +104,11 @@ void SDataprepGraphActionStepNode::UpdateGraphNode()
 				UDataprepFilter* Filter = static_cast<UDataprepFilter*>( StepObject );
 				ActionStepBlockPtr = StaticCastSharedRef<SDataprepActionBlock>( SNew(SDataprepFilter, *Filter, StepData) );
 			}
+			else if (StepType == UDataprepSelectionTransform::StaticClass())
+			{
+				UDataprepSelectionTransform* SelectionTransform = static_cast<UDataprepSelectionTransform*>( StepObject );
+				ActionStepBlockPtr = StaticCastSharedRef<SDataprepActionBlock>( SNew(SDataprepSelectionTransform, SelectionTransform, StepData) );
+			}
 
 			if(ActionStepBlockPtr.IsValid())
 			{
diff --git a/Engine/Plugins/Enterprise/DataprepEditor/Source/DataprepEditor/Private/Widgets/DataprepGraph/SDataprepSelectionTransform.cpp b/Engine/Plugins/Enterprise/DataprepEditor/Source/DataprepEditor/Private/Widgets/DataprepGraph/SDataprepSelectionTransform.cpp
new file mode 100644
--- /dev/null
+++ b/Engine/Plugins/Enterprise/DataprepEditor/Source/DataprepEditor/Private/Widgets/DataprepGraph/SDataprepSelectionTransform.cpp
@@ -0,0 +1,52 @@
+// Copyright Epic Games, Inc. All Rights Reserved.
+
+#include "SDataprepSelectionTransform.h"
+
+#include "DataprepEditorStyle.h"
+#include "SelectionSystem/DataprepSelectionTransform.h"
+#include "Widgets/DataprepWidgets.h"
+#include "Widgets/SNullWidget.h"
+
+#include "IDetailsView.h"
+#include "Modules/ModuleManager.h"
+#include "PropertyEditorModule.h"
+
+void SDataprepSelectionTransform::Construct(const FArguments& InArgs, UDataprepSelectionTransform* InTransform, const TSharedRef<FDataprepSchemaActionContext>& InDataprepActionContext)
+{
+	SelectionTransform = InTransform;
+
+	TAttribute<FText> TooltipTextAttribute = MakeAttributeSP(this, &SDataprepSelectionTransform::GetTooltipText);
+	SetToolTipText(TooltipTextAttribute);
+
+	SDataprepActionBlock::Construct(SDataprepActionBlock::FArguments(), InDataprepActionContext);
+}
+
+FLinearColor SDataprepSelectionTransform::GetOutlineColor() const
+{
+	return FDataprepEditorStyle::GetColor("DataprepActionStep.Operation.OutlineColor");
+}
+
+FText SDataprepSelectionTransform::GetBlockTitle() const
+{
+	return SelectionTransform ? SelectionTransform->GetDisplayTransformName() : FText::FromString(TEXT("Transform is Nullptr!"));
+}
+
+TSharedRef<SWidget> SDataprepSelectionTransform::GetContentWidget()
+{
+	return SNew(SDataprepDetailsView).Object(SelectionTransform);
+}
+
+void SDataprepSelectionTransform::AddReferencedObjects(FReferenceCollector& Collector)
+{
+	Collector.AddReferencedObject(SelectionTransform);
+}
+
+FText SDataprepSelectionTransform::GetTooltipText() const
+{
+	FText TooltipText;
+	if (SelectionTransform)
+	{
+		TooltipText = SelectionTransform->GetTooltip();
+	}
+	return TooltipText;
+}
diff --git a/Engine/Plugins/Enterprise/DataprepEditor/Source/DataprepEditor/Private/Widgets/DataprepGraph/SDataprepSelectionTransform.h b/Engine/Plugins/Enterprise/DataprepEditor/Source/DataprepEditor/Private/Widgets/DataprepGraph/SDataprepSelectionTransform.h
new file mode 100644
--- /dev/null
+++ b/Engine/Plugins/Enterprise/DataprepEditor/Source/DataprepEditor/Private/Widgets/DataprepGraph/SDataprepSelectionTransform.h
@@ -0,0 +1,38 @@
+// Copyright Epic Games, Inc. All Rights Reserved.
+
+#pragma once
+
+#include "Widgets/DataprepGraph/SDataprepActionBlock.h"
+
+#include "CoreMinimal.h"
+#include "UObject/GCObject.h"
+#include "Widgets/DeclarativeSyntaxSupport.h"
+#include "Widgets/SCompoundWidget.h"
+
+class UDataprepSelectionTransform;
+struct FDataprepSchemaActionContext;
+
+class SDataprepSelectionTransform : public SDataprepActionBlock, public FGCObject
+{
+	SLATE_BEGIN_ARGS(SDataprepSelectionTransform) {}
+	SLATE_END_ARGS();
+
+	void Construct(const FArguments& InArgs, UDataprepSelectionTransform* InTransform, const TSharedRef<FDataprepSchemaActionContext>& InDataprepActionContext);
+
+protected:
+
+	// SDataprepActionBlock interface
+	virtual FLinearColor GetOutlineColor() const override;
+	virtual FText GetBlockTitle() const override;
+	virtual TSharedRef<SWidget> GetContentWidget() override;
+	//~ end of SDataprepActionBlock interface
+
+private:
+
+	FText GetTooltipText() const;
+
+	//~ FGCObject interface
+	virtual void AddReferencedObjects(FReferenceCollector& Collector) override;
+
+	UDataprepSelectionTransform* SelectionTransform;
+};
diff --git a/Engine/Plugins/Enterprise/DataprepEditor/Source/DataprepLibraries/Private/DataprepSelectionTransforms.cpp b/Engine/Plugins/Enterprise/DataprepEditor/Source/DataprepLibraries/Private/DataprepSelectionTransforms.cpp
new file mode 100644
--- /dev/null
+++ b/Engine/Plugins/Enterprise/DataprepEditor/Source/DataprepLibraries/Private/DataprepSelectionTransforms.cpp
@@ -0,0 +1,88 @@
+// Copyright Epic Games, Inc. All Rights Reserved.
+
+#include "DataprepSelectionTransforms.h"
+#include "Engine/StaticMesh.h"
+#include "Engine/StaticMeshActor.h"
+#include "Engine/Texture.h"
+#include "Materials/MaterialInstance.h"
+
+#define LOCTEXT_NAMESPACE "DataprepSelectionTransforms"
+
+void UDataprepReferenceSelectionTransform::OnExecution_Implementation(const TArray<UObject*>& InObjects, TArray<UObject*>& OutObjects)
+{
+	TSet<UObject*> Assets;
+
+	for (UObject* Object : InObjects)
+	{
+		if (!ensure(Object) || Object->IsPendingKill())
+		{
+			continue;
+		}
+
+		if (AActor* Actor = Cast< AActor >(Object))
+		{
+			TArray<UActorComponent*> Components = Actor->GetComponents().Array();
+			Components.Append(Actor->GetInstanceComponents());
+
+			for (UActorComponent* Component : Components)
+			{
+				if (UStaticMeshComponent* MeshComponent = Cast<UStaticMeshComponent>(Component))
+				{
+					if (UStaticMesh* StaticMesh = MeshComponent->GetStaticMesh())
+					{
+						Assets.Add(StaticMesh);
+					}
+
+					for (UMaterialInterface* MaterialInterface : MeshComponent->OverrideMaterials)
+					{
+						Assets.Add(MaterialInterface);
+					}
+				}
+			}
+		}
+		else if (UStaticMesh* StaticMesh = Cast< UStaticMesh >(Object))
+		{
+			for (FStaticMaterial& StaticMaterial : StaticMesh->StaticMaterials)
+			{
+				if (UMaterialInterface* MaterialInterface = StaticMaterial.MaterialInterface)
+				{
+					Assets.Add(MaterialInterface);
+				}
+			}
+
+			if (bOutputCanIncludeInput)
+			{
+				Assets.Add(Object);
+			}
+		}
+		else if (UMaterialInterface* MaterialInterface = Cast< UMaterialInterface >(Object))
+		{
+			Assets.Add(MaterialInterface);
+
+			if (UMaterialInstance* MaterialInstance = Cast<UMaterialInstance>(MaterialInterface))
+			{
+				if (MaterialInstance->Parent)
+				{
+					Assets.Add(MaterialInstance->Parent);
+				}
+			}
+
+			// Collect textures
+			TArray<UTexture*> Textures;
+			MaterialInterface->GetUsedTextures(Textures, EMaterialQualityLevel::Num, true, ERHIFeatureLevel::Num, true);
+			for (UTexture* Texture : Textures)
+			{
+				Assets.Add(Texture);
+			}
+
+			if (bOutputCanIncludeInput)
+			{
+				Assets.Add(Object);
+			}
+		}
+	}
+
+	OutObjects.Append(Assets.Array());
+}
+
+#undef LOCTEXT_NAMESPACE
diff --git a/Engine/Plugins/Enterprise/DataprepEditor/Source/DataprepLibraries/Private/DataprepSelectionTransforms.h b/Engine/Plugins/Enterprise/DataprepEditor/Source/DataprepLibraries/Private/DataprepSelectionTransforms.h
new file mode 100644
--- /dev/null
+++ b/Engine/Plugins/Enterprise/DataprepEditor/Source/DataprepLibraries/Private/DataprepSelectionTransforms.h
@@ -0,0 +1,22 @@
+// Copyright Epic Games, Inc. All Rights Reserved.
+
+#pragma once
+
+#include "SelectionSystem/DataprepSelectionTransform.h"
+
+#include "CoreMinimal.h"
+#include "Delegates/DelegateCombinations.h"
+#include "UObject/Object.h"
+#include "UObject/ObjectMacros.h"
+#include "UObject/TextProperty.h"
+
+#include "DataprepSelectionTransforms.generated.h"
+
+UCLASS(Category = SelectionTransform, Meta = (DisplayName="Reference Selection Transform", ToolTip = "Return all the assets used/referenced by the selected objects") )
+class UDataprepReferenceSelectionTransform : public UDataprepSelectionTransform
+{
+	GENERATED_BODY()
+
+protected:
+	virtual void OnExecution_Implementation(const TArray<UObject*>& InObjects, TArray<UObject*>& OutObjects) override;
+};
