Integrating CL #11335522 from //UE4/Dev-Enterprise

Implement 'Spawn Actors at Location' operation

diff --git a/Engine/Plugins/Enterprise/DataprepEditor/Source/DataprepLibraries/Private/DataprepEditingOperations.cpp b/Engine/Plugins/Enterprise/DataprepEditor/Source/DataprepLibraries/Private/DataprepEditingOperations.cpp
--- a/Engine/Plugins/Enterprise/DataprepEditor/Source/DataprepLibraries/Private/DataprepEditingOperations.cpp
+++ b/Engine/Plugins/Enterprise/DataprepEditor/Source/DataprepLibraries/Private/DataprepEditingOperations.cpp
@@ -608,6 +608,42 @@ void UDataprepCompactSceneGraphOperation::OnExecution_Implementation(const FData
 	DeleteObjects(ObjectsToDelete);
 }
 
+void UDataprepSpawnActorsAtLocation::OnExecution_Implementation(const FDataprepContext& InContext)
+{
+#ifdef LOG_TIME
+	DataprepEditingOperationTime::FTimeLogger TimeLogger(TEXT("SpawnActorsAtLocation"), [&](FText Text) { this->LogInfo(Text); });
+#endif
+
+	if (!StaticMesh)
+	{
+		UE_LOG(LogDataprep, Log, TEXT("No mesh was selected"));
+		return;
+	}
+
+	for (UObject* Object : InContext.Objects)
+	{
+		if (!ensure(Object) || Object->IsPendingKill())
+		{
+			continue;
+		}
+
+		if (AActor* Actor = Cast<AActor>(Object))
+		{
+			AStaticMeshActor* NewActor = Cast<AStaticMeshActor>( CreateActor(AStaticMeshActor::StaticClass(), Actor->GetName()) );
+			check(NewActor);
+
+			NewActor->GetStaticMeshComponent()->SetStaticMesh(StaticMesh);
+			NewActor->SetActorTransform(Actor->GetActorTransform());
+			NewActor->SetActorLabel(Actor->GetActorLabel() + TEXT("_SA"));
+
+			if (Actor->GetParentActor())
+			{
+				NewActor->AttachToActor(Actor->GetParentActor(), FAttachmentTransformRules::KeepRelativeTransform);
+			}
+		}
+	}
+}
+
 bool UDataprepCompactSceneGraphOperation::IsActorVisible(AActor* Actor, TMap<AActor*, bool>& VisibilityMap)
 {
 	if (!Actor)
diff --git a/Engine/Plugins/Enterprise/DataprepEditor/Source/DataprepLibraries/Private/DataprepEditingOperations.h b/Engine/Plugins/Enterprise/DataprepEditor/Source/DataprepLibraries/Private/DataprepEditingOperations.h
--- a/Engine/Plugins/Enterprise/DataprepEditor/Source/DataprepLibraries/Private/DataprepEditingOperations.h
+++ b/Engine/Plugins/Enterprise/DataprepEditor/Source/DataprepLibraries/Private/DataprepEditingOperations.h
@@ -157,3 +157,23 @@ protected:
 
 	bool IsActorVisible(AActor*, TMap<AActor*, bool>& VisibilityMap);
 };
+
+UCLASS(Experimental, Category = ActorOperation, Meta = (DisplayName = "Spawn Actors At Location", ToolTip = "For each actor in the input set, spawn an actor from the specified Asset at the same position and orientation than the reference"))
+class UDataprepSpawnActorsAtLocation : public UDataprepEditingOperation
+{
+	GENERATED_BODY()
+
+	//~ Begin UDataprepOperation Interface
+public:
+	virtual FText GetCategory_Implementation() const override
+	{
+		return FDataprepOperationCategories::ActorOperation;
+	}
+
+	UPROPERTY(EditAnywhere, BlueprintReadWrite, Category = ActorOperation)
+	UStaticMesh* StaticMesh;
+
+protected:
+	virtual void OnExecution_Implementation(const FDataprepContext& InContext) override;
+	//~ End UDataprepOperation Interface
+};
