Fixed crash when using "replace asset references" operation
Replaced weak UObject pointers by soft object path when caching assets created by the Dataprep. This improves the tracking of those assets if operations transform them.
Moved Dataprep producers to use one package by asset
Modified Datasmith importer to support assets which are not sub-assets

diff --git a/Engine/Plugins/Enterprise/DataprepEditor/Source/DataprepCore/Private/DataprepContentProducer.cpp b/Engine/Plugins/Enterprise/DataprepEditor/Source/DataprepCore/Private/DataprepContentProducer.cpp
--- a/Engine/Plugins/Enterprise/DataprepEditor/Source/DataprepCore/Private/DataprepContentProducer.cpp
+++ b/Engine/Plugins/Enterprise/DataprepEditor/Source/DataprepCore/Private/DataprepContentProducer.cpp
@@ -55,15 +55,27 @@ bool UDataprepContentProducer::Produce(const FDataprepProducerContext& InContext
 	TSet< UPackage* > PackagesToCheck;
 	for(int32 Index = LastAssetCount; Index < OutAssets.Num(); ++Index)
 	{
-		UObject* Asset = OutAssets[Index].Get();
-
-		if( ULevelSequence* LevelSequence = Cast<ULevelSequence>(Asset) )
-		{
-			PackagesToCheck.Add( LevelSequence->GetOutermost() );
-		}
-		else if( ULevelVariantSets* ThisLevelVariantSets = Cast<ULevelVariantSets>(Asset) )
+		if(UObject* Asset = OutAssets[Index].Get())
 		{
-			PackagesToCheck.Add( ThisLevelVariantSets->GetOutermost() );
+			// Verify each asset is within its own package
+			FString PackageName = Asset->GetOutermost()->GetName();
+			FString AssetName = Asset->GetName();
+			if(AssetName != PackageName)
+			{
+				UPackage* NewPackage = NewObject<UPackage>( nullptr, *FPaths::Combine( PackageName, AssetName ), Asset->GetOutermost()->GetFlags() );
+				NewPackage->FullyLoad();
+
+				Asset->Rename( nullptr, NewPackage, REN_NonTransactional | REN_DontCreateRedirectors );
+			}
+
+			if( ULevelSequence* LevelSequence = Cast<ULevelSequence>(Asset) )
+			{
+				PackagesToCheck.Add( LevelSequence->GetOutermost() );
+			}
+			else if( ULevelVariantSets* ThisLevelVariantSets = Cast<ULevelVariantSets>(Asset) )
+			{
+				PackagesToCheck.Add( ThisLevelVariantSets->GetOutermost() );
+			}
 		}
 	}
 
diff --git a/Engine/Plugins/Enterprise/DataprepEditor/Source/DataprepEditor/Private/DataprepEditor.cpp b/Engine/Plugins/Enterprise/DataprepEditor/Source/DataprepEditor/Private/DataprepEditor.cpp
--- a/Engine/Plugins/Enterprise/DataprepEditor/Source/DataprepEditor/Private/DataprepEditor.cpp
+++ b/Engine/Plugins/Enterprise/DataprepEditor/Source/DataprepEditor/Private/DataprepEditor.cpp
@@ -568,8 +568,14 @@ void FDataprepEditor::OnBuildWorld()
 		return;
 	}
 
-	CachedAssets.Reset();
-	CachedAssets.Append( Assets );
+	CachedAssets.Empty(Assets.Num());
+	for(TWeakObjectPtr<UObject>& WeakObject : Assets)
+	{
+		if(UObject* Object = WeakObject.Get())
+		{
+			CachedAssets.Emplace( Object );
+		}
+	}
 
 	TakeSnapshot();
 
@@ -644,16 +650,17 @@ void FDataprepEditor::CleanPreviewWorld()
 	// Delete assets which are still in the transient content folder
 	const FString TransientContentFolder( GetTransientContentFolder() );
 	TArray<UObject*> ObjectsToDelete;
-	for( TWeakObjectPtr<UObject>& Asset : CachedAssets )
+	for( FSoftObjectPath& SoftObjectPath : CachedAssets )
 	{
-		if( UObject* ObjectToDelete = Asset.Get() )
+		if(SoftObjectPath.GetLongPackageName().StartsWith(TransientContentFolder))
 		{
-			FString PackagePath = ObjectToDelete->GetOutermost()->GetName();
-			if( PackagePath.StartsWith( TransientContentFolder ) )
-			{
-				FDataprepCoreUtils::MoveToTransientPackage( ObjectToDelete );
-				ObjectsToDelete.Add( ObjectToDelete );
-			}
+			FSoftObjectPath PackagePath(SoftObjectPath.GetLongPackageName());
+			UPackage* Package = Cast<UPackage>(PackagePath.ResolveObject());
+
+			UObject* ObjectToDelete = StaticFindObjectFast(nullptr, Package, *SoftObjectPath.GetAssetName());
+
+			FDataprepCoreUtils::MoveToTransientPackage( ObjectToDelete );
+			ObjectsToDelete.Add( ObjectToDelete );
 		}
 	}
 
@@ -726,7 +733,7 @@ void FDataprepEditor::OnExecutePipeline()
 	{
 		if( Asset.IsValid() )
 		{
-			CachedAssets.Add( Asset );
+			CachedAssets.Emplace( Asset.Get() );
 		}
 	}
 
@@ -781,6 +788,11 @@ void FDataprepEditor::OnCommitWorld()
 	if( !DataprepAssetInterfacePtr->RunConsumer( Context ) )
 	{
 		UE_LOG( LogDataprepEditor, Error, TEXT("Consumer failed...") );
+
+		// Restore 3D viewport
+		SceneViewportView->UpdateScene();
+
+		return;
 	}
 
 	ResetBuildWorld();
diff --git a/Engine/Plugins/Enterprise/DataprepEditor/Source/DataprepEditor/Private/DataprepEditor.h b/Engine/Plugins/Enterprise/DataprepEditor/Source/DataprepEditor/Private/DataprepEditor.h
--- a/Engine/Plugins/Enterprise/DataprepEditor/Source/DataprepEditor/Private/DataprepEditor.h
+++ b/Engine/Plugins/Enterprise/DataprepEditor/Source/DataprepEditor/Private/DataprepEditor.h
@@ -226,7 +226,7 @@ private:
 	 * All assets tracked for this editor.
 	 */
 	TArray<TWeakObjectPtr<UObject>> Assets;
-	TSet<TWeakObjectPtr<UObject>> CachedAssets;
+	TSet<FSoftObjectPath> CachedAssets;
 
 	/**
 	 * The world used to preview the inputs
diff --git a/Engine/Plugins/Enterprise/DataprepEditor/Source/DataprepEditor/Private/Widgets/SAssetsPreviewWidget.cpp b/Engine/Plugins/Enterprise/DataprepEditor/Source/DataprepEditor/Private/Widgets/SAssetsPreviewWidget.cpp
--- a/Engine/Plugins/Enterprise/DataprepEditor/Source/DataprepEditor/Private/Widgets/SAssetsPreviewWidget.cpp
+++ b/Engine/Plugins/Enterprise/DataprepEditor/Source/DataprepEditor/Private/Widgets/SAssetsPreviewWidget.cpp
@@ -384,6 +384,17 @@ namespace AssetPreviewWidget
 		}
 		AssetSubPath.ReplaceCharInline(TEXT('/'), TEXT('.'));
 		AssetSubPath.ParseIntoArray(ItemsName, TEXT("."), true);
+
+		// Check that asset's name is not repeated twice at the end
+		if(ItemsName.Num() > 1)
+		{
+			const int32 LastIndex = ItemsName.Num() - 1;
+			if(ItemsName[LastIndex] == ItemsName[LastIndex - 1])
+			{
+				ItemsName.RemoveAt(LastIndex, 1, false);
+			}
+		}
+
 		return ItemsName;
 	}
 
diff --git a/Engine/Plugins/Enterprise/DatasmithImporter/Source/DatasmithImporter/Private/DatasmithImporter.cpp b/Engine/Plugins/Enterprise/DatasmithImporter/Source/DatasmithImporter/Private/DatasmithImporter.cpp
--- a/Engine/Plugins/Enterprise/DatasmithImporter/Source/DatasmithImporter/Private/DatasmithImporter.cpp
+++ b/Engine/Plugins/Enterprise/DatasmithImporter/Source/DatasmithImporter/Private/DatasmithImporter.cpp
@@ -110,8 +110,10 @@ namespace DatasmithImporterImpl
 
 		if ( !ExistingAsset )
 		{
-			FString DestinationPackagePath = UPackageTools::SanitizePackageName( FPaths::Combine( DestinationPath, SourceAsset->GetName() ) );
-			FString DestinationAssetPath = DestinationPackagePath + TEXT(".") + UPackageTools::SanitizePackageName( SourceAsset->GetName() );
+			const FString AssetName = SourceAsset->GetName();
+			bool bPathIsComplete = AssetName == FPaths::GetBaseFilename( DestinationPath );
+			FString DestinationPackagePath = UPackageTools::SanitizePackageName( bPathIsComplete ? DestinationPath : FPaths::Combine( DestinationPath, AssetName ) );
+			FString DestinationAssetPath = DestinationPackagePath + TEXT(".") + UPackageTools::SanitizePackageName( AssetName );
 
 			ExistingAsset = FDatasmithImporterUtils::FindObject<UObject>( nullptr, DestinationAssetPath );
 
