Complementary submission to fix crash: Added sorting of objects to consolidate to consolidate in the right order. From less dependent to more dependent.

diff --git a/Engine/Plugins/Enterprise/DataprepEditor/Source/DataprepLibraries/Private/DataprepOperationsLibrary.cpp b/Engine/Plugins/Enterprise/DataprepEditor/Source/DataprepLibraries/Private/DataprepOperationsLibrary.cpp
--- a/Engine/Plugins/Enterprise/DataprepEditor/Source/DataprepLibraries/Private/DataprepOperationsLibrary.cpp
+++ b/Engine/Plugins/Enterprise/DataprepEditor/Source/DataprepLibraries/Private/DataprepOperationsLibrary.cpp
@@ -22,6 +22,7 @@
 #include "Interfaces/ITargetPlatform.h"
 #include "Interfaces/ITargetPlatformManagerModule.h"
 #include "Materials/Material.h"
+#include "Materials/MaterialFunctionInstance.h"
 #include "Materials/MaterialInterface.h"
 #include "Math/Vector2D.h"
 #include "Misc/FileHelper.h"
@@ -64,7 +65,6 @@ namespace DataprepOperationsLibraryUtil
 	TSet<UStaticMesh*> GetSelectedMeshes(const TArray<UObject*>& SelectedObjects)
 	{
 		TSet<UStaticMesh*> SelectedMeshes;
-		TArray<AActor*> SelectedActors;
 
 		for (UObject* Object : SelectedObjects)
 		{
@@ -786,8 +786,47 @@ void UDataprepOperationsLibrary::ConsolidateObjects(const TArray< UObject* >& Se
 		OutCompatibleObjects.Add(CurProposedObj);
 	}
 
+	// Sort assets according to their dependency
+	// Texture first, then MaterialFunction, then ...
+	auto GetAssetClassRank = [&](const UClass* AssetClass) -> int8
+	{
+		if (AssetClass->IsChildOf(UTexture::StaticClass()))
+		{
+			return 0;
+		}
+		else if (AssetClass->IsChildOf(UMaterialFunction::StaticClass()))
+		{
+			return 1;
+		}
+		else if (AssetClass->IsChildOf(UMaterialFunctionInstance::StaticClass()))
+		{
+			return 2;
+		}
+		else if (AssetClass->IsChildOf(UMaterial::StaticClass()))
+		{
+			return 3;
+		}
+		else if (AssetClass->IsChildOf(UMaterialInstance::StaticClass()))
+		{
+			return 4;
+		}
+		else if (AssetClass->IsChildOf(UStaticMesh::StaticClass()))
+		{
+			return 5;
+		}
+
+		return 6;
+	};
+
+	Algo::Sort(OutCompatibleObjects, [&](const UObject* A, const UObject* B)
+	{
+		int8 AValue = A ? GetAssetClassRank(A->GetClass()) : 7;
+		int8 BValue = B ? GetAssetClassRank(B->GetClass()) : 7;
+		return AValue > BValue;
+	});
+
 	// Perform the object consolidation
-	ObjectTools::FConsolidationResults ConsResults = ObjectTools::ConsolidateObjects(ObjectToConsolidateTo, OutCompatibleObjects, false);
+	ObjectTools::ConsolidateObjects(ObjectToConsolidateTo, OutCompatibleObjects, false);
 }
 
 void UDataprepOperationsLibrary::RandomizeTransform(const TArray<UObject*>& SelectedObjects, ERandomizeTransformType TransformType, ERandomizeTransformReferenceFrame ReferenceFrame, const FVector& Min, const FVector& Max)
