Fixed crash when taking a snapshot of big meshes

diff --git a/Engine/Plugins/Enterprise/DataprepEditor/Source/DataprepEditor/Private/DataprepSnapshot.cpp b/Engine/Plugins/Enterprise/DataprepEditor/Source/DataprepEditor/Private/DataprepSnapshot.cpp
--- a/Engine/Plugins/Enterprise/DataprepEditor/Source/DataprepEditor/Private/DataprepSnapshot.cpp
+++ b/Engine/Plugins/Enterprise/DataprepEditor/Source/DataprepEditor/Private/DataprepSnapshot.cpp
@@ -51,8 +51,119 @@ enum class EDataprepAssetClass : uint8 {
 
 namespace DataprepSnapshotUtil
 {
+	typedef TArray<uint8, TSizedDefaultAllocator<64>> FSnapshotBuffer;
+
 	const TCHAR* SnapshotExtension = TEXT(".dpc");
 
+	// #ueent_todo: Remove those temporary fixes for large data saving/loading with implementation from Core
+	class FSnapshotWriter : public FMemoryArchive
+	{
+	public:
+		FSnapshotWriter( FSnapshotBuffer& InBytes, bool bIsPersistent = false, bool bSetOffset = false, const FName InArchiveName = NAME_None )
+			: FMemoryArchive()
+			, Bytes(InBytes)
+			, ArchiveName(InArchiveName)
+		{
+			this->SetIsSaving(true);
+			this->SetIsPersistent(bIsPersistent);
+			if (bSetOffset)
+			{
+				Offset = InBytes.Num();
+			}
+		}
+
+		virtual void Serialize(void* Data, int64 Num) override
+		{
+			const int64 NumBytesToAdd = Offset + Num - Bytes.Num();
+			if( NumBytesToAdd > 0 )
+			{
+				const int64 NewArrayCount = Bytes.Num() + NumBytesToAdd;
+				Bytes.AddUninitialized( NumBytesToAdd );
+			}
+
+			check((Offset + Num) <= Bytes.Num());
+
+			if( Num )
+			{
+				FMemory::Memcpy( &Bytes[Offset], Data, Num );
+				Offset+=Num;
+			}
+		}
+		/**
+		* Returns the name of the Archive.  Useful for getting the name of the package a struct or object
+		* is in when a loading error occurs.
+		*
+		* This is overridden for the specific Archive Types
+		**/
+		virtual FString GetArchiveName() const override { return TEXT("FSnapshotWriter"); }
+
+		int64 TotalSize() override
+		{
+			return Bytes.Num();
+		}
+
+	protected:
+
+		FSnapshotBuffer&	Bytes;
+
+		/** Archive name, used to debugging, by default set to NAME_None. */
+		const FName ArchiveName;
+	};
+
+	class FSnapshotReader : public FMemoryArchive
+	{
+	public:
+		/**
+		* Returns the name of the Archive.  Useful for getting the name of the package a struct or object
+		* is in when a loading error occurs.
+		*
+		* This is overridden for the specific Archive Types
+		**/
+		virtual FString GetArchiveName() const override
+		{
+			return TEXT("FSnapshotReader");
+		}
+
+		virtual int64 TotalSize() override
+		{
+			return FMath::Min(Bytes.Num(), LimitSize);
+		}
+
+		void Serialize( void* Data, int64 Num )
+		{
+			if (Num && !ArIsError)
+			{
+				// Only serialize if we have the requested amount of data
+				if (Offset + Num <= TotalSize())
+				{
+					FMemory::Memcpy( Data, &Bytes[Offset], Num );
+					Offset += Num;
+				}
+				else
+				{
+					ArIsError = true;
+				}
+			}
+		}
+
+		explicit FSnapshotReader( const FSnapshotBuffer& InBytes, bool bIsPersistent = false )
+			: Bytes    (InBytes)
+			, LimitSize(INT64_MAX)
+		{
+			this->SetIsLoading(true);
+			this->SetIsPersistent(bIsPersistent);
+		}
+
+		/** With this method it's possible to attach data behind some serialized data. */
+		void SetLimitSize(int64 NewLimitSize)
+		{
+			LimitSize = NewLimitSize;
+		}
+
+	private:
+		const FSnapshotBuffer&	Bytes;
+		int64					LimitSize;
+	};
 
 	/**
 	 * Extends FObjectAndNameAsStringProxyArchive to support FLazyObjectPtr.
@@ -102,7 +213,7 @@ namespace DataprepSnapshotUtil
 		return FPaths::ConvertRelativePathToFull( FPaths::Combine( RootPath, PackageFileName ) + SnapshotExtension );
 	}
 
-	void WriteSnapshotData(UObject* Object, TArray<uint8>& OutSerializedData)
+	void WriteSnapshotData(UObject* Object, FSnapshotBuffer& OutSerializedData)
 	{
 		// Helper struct to identify dependency of a UObject on other UObject(s) except given one (its outer)
 		struct FObjectDependencyAnalyzer : public FArchiveUObject
@@ -136,7 +247,7 @@ namespace DataprepSnapshotUtil
 			TSet<UObject*> DependentObjects;
 		};
 
-		FMemoryWriter MemAr(OutSerializedData);
+		FSnapshotWriter MemAr(OutSerializedData);
 		FSnapshotCustomArchive Ar(MemAr);
 
 		// Collect sub-objects depending on input object including nested objects
@@ -233,6 +344,9 @@ namespace DataprepSnapshotUtil
 		}
 
 		// Serialize object
+		// Overwrite transacting and persistent flags if Object is a static mesh. This will skip the persistence of the MeshDescription
+		Ar.SetIsTransacting(Cast<UStaticMesh>(Object) == nullptr);
+		Ar.SetIsPersistent(Cast<UStaticMesh>(Object) != nullptr);
 		Object->Serialize(Ar);
 
 		if(UTexture* Texture = Cast<UTexture>(Object))
@@ -242,7 +356,7 @@ namespace DataprepSnapshotUtil
 		}
 	}
 
-	void ReadSnapshotData(UObject* Object, const TArray<uint8>& InSerializedData, TMap<FString, UClass*>& InClassesMap, TArray<UObject*>& ObjectsToDelete)
+	void ReadSnapshotData(UObject* Object, const FSnapshotBuffer& InSerializedData, TMap<FString, UClass*>& InClassesMap, TArray<UObject*>& ObjectsToDelete)
 	{
 		// Remove all objects created by default that InObject is dependent on
 		// This method must obviously be called just after the InObject is created
@@ -260,7 +374,7 @@ namespace DataprepSnapshotUtil
 
 		RemoveDefaultDependencies( Object );
 
-		FMemoryReader MemAr(InSerializedData);
+		FSnapshotReader MemAr(InSerializedData);
 		FSnapshotCustomArchive Ar(MemAr);
 
 		// Deserialize count of sub-objects
@@ -317,6 +431,9 @@ namespace DataprepSnapshotUtil
 		}
 
 		// Deserialize object
+		// Overwrite transacting and persistent flags if Object is a static mesh. This will skip the persistence of the MeshDescription
+		Ar.SetIsTransacting(Cast<UStaticMesh>(Object) == nullptr);
+		Ar.SetIsPersistent(Cast<UStaticMesh>(Object) != nullptr);
 		Object->Serialize(Ar);
 
 		if(UTexture* Texture = Cast<UTexture>(Object))
@@ -510,7 +627,7 @@ void FDataprepEditor::TakeSnapshot()
 
 								// Serialize asset
 								{
-									TArray<uint8> SerializedData;
+									DataprepSnapshotUtil::FSnapshotBuffer SerializedData;
 									DataprepSnapshotUtil::WriteSnapshotData( AssetObject, SerializedData );
 
 									FString AssetFilePath = DataprepSnapshotUtil::BuildAssetFileName( this->TempDir, AssetPathString );
@@ -679,7 +796,7 @@ void FDataprepEditor::RestoreFromSnapshot(bool bUpdateViewport)
 			UObject* Asset = NewObject<UObject>( Package, DataEntry.Get<1>(), *AssetName, DataEntry.Get<2>() );
 
 			{
-				TArray<uint8> SerializedData;
+				DataprepSnapshotUtil::FSnapshotBuffer SerializedData;
 				FString AssetFilePath = DataprepSnapshotUtil::BuildAssetFileName( TempDir, ObjectPath.GetAssetPathString() );
 				if(FFileHelper::LoadFileToArray(SerializedData, *AssetFilePath))
 				{
