Fixed issue with translucent materials in Dataprep editor's viewport after import

diff --git a/Engine/Plugins/Enterprise/DataprepEditor/Source/DataprepCore/DataprepCore.Build.cs b/Engine/Plugins/Enterprise/DataprepEditor/Source/DataprepCore/DataprepCore.Build.cs
--- a/Engine/Plugins/Enterprise/DataprepEditor/Source/DataprepCore/DataprepCore.Build.cs
+++ b/Engine/Plugins/Enterprise/DataprepEditor/Source/DataprepCore/DataprepCore.Build.cs
@@ -25,6 +25,7 @@ namespace UnrealBuildTool.Rules
 					"MessageLog",
 					"PropertyEditor",
 					"RenderCore",
+					"RHI",
 					"Slate",
 					"SlateCore",
 					"StaticMeshDescription",
diff --git a/Engine/Plugins/Enterprise/DataprepEditor/Source/DataprepCore/Private/DataprepCoreUtils.cpp b/Engine/Plugins/Enterprise/DataprepEditor/Source/DataprepCore/Private/DataprepCoreUtils.cpp
--- a/Engine/Plugins/Enterprise/DataprepEditor/Source/DataprepCore/Private/DataprepCoreUtils.cpp
+++ b/Engine/Plugins/Enterprise/DataprepEditor/Source/DataprepCore/Private/DataprepCoreUtils.cpp
@@ -30,6 +30,7 @@
 #include "Materials/Material.h"
 #include "Materials/MaterialFunction.h"
 #include "Materials/MaterialInstance.h"
+#include "Materials/MaterialInstanceConstant.h"
 #include "Misc/ScopedSlowTask.h"
 #include "RenderingThread.h"
 #include "UObject/StrongObjectPtr.h"
@@ -617,15 +618,43 @@ void FDataprepCoreUtils::BuildAssets(const TArray<TWeakObjectPtr<UObject>>& Asse
 	// Force compilation of materials which have no render proxy
 	if(MaterialInterfaces.Num() > 0)
 	{
-		auto MustCompile = [](const UMaterialInterface* MaterialInterface)
+		auto MustCompile = [](UMaterialInterface* MaterialInterface) -> bool
 		{
-			if( MaterialInterface )
+			if(MaterialInterface == nullptr)
 			{
-				const FMaterialRenderProxy* RenderProxy = MaterialInterface->GetRenderProxy();
-				return RenderProxy == nullptr || !RenderProxy->IsInitialized();
+				return false;
 			}
 
-			return false;
+			// Force recompilation of constant material instances which either override blend mode or any static switch
+			if(UMaterialInstanceConstant* ConstantMaterialInstance = Cast< UMaterialInstanceConstant >(MaterialInterface))
+			{
+				// If BlendMode override property has been changed, make sure this combination of the parent material is compiled
+				if ( ConstantMaterialInstance->BasePropertyOverrides.bOverride_BlendMode == true )
+				{
+					ConstantMaterialInstance->ForceRecompileForRendering();
+					return true;
+				}
+				else
+				{
+					// If a static switch is overridden, we need to recompile
+					FStaticParameterSet StaticParameters;
+					ConstantMaterialInstance->GetStaticParameterValues( StaticParameters );
+
+					for ( FStaticSwitchParameter& Switch : StaticParameters.StaticSwitchParameters )
+					{
+						if ( Switch.bOverride )
+						{
+							ConstantMaterialInstance->ForceRecompileForRendering();
+							return true;
+						}
+					}
+				}
+			}
+
+			// Force compilation if there is no valid render proxy
+			const FMaterialRenderProxy* RenderProxy = MaterialInterface->GetRenderProxy();
+
+			return RenderProxy == nullptr || !RenderProxy->IsInitialized();
 		};
 
 		int32 AssetBuiltCount = 0;
diff --git a/Engine/Plugins/Enterprise/DataprepEditor/Source/DataprepEditor/Private/Widgets/SDataprepEditorViewport.cpp b/Engine/Plugins/Enterprise/DataprepEditor/Source/DataprepEditor/Private/Widgets/SDataprepEditorViewport.cpp
--- a/Engine/Plugins/Enterprise/DataprepEditor/Source/DataprepEditor/Private/Widgets/SDataprepEditorViewport.cpp
+++ b/Engine/Plugins/Enterprise/DataprepEditor/Source/DataprepEditor/Private/Widgets/SDataprepEditorViewport.cpp
@@ -171,7 +171,7 @@ namespace DataprepEditor3DPreviewUtils
 
 	/** Compile all materials included in the input array*/
 	// Copied from DatasmithImporterImpl::CompileMaterial
-	void CompileMaterials(const TArray< UMaterialInterface* > Materials);
+	void CompileMaterials(const TArray< UMaterialInterface* >& Materials);
 
 	void FindMeshComponents(const AActor * InActor, TArray<UStaticMeshComponent*>& MeshComponents, bool bRecursive );
 
@@ -721,6 +721,21 @@ void SDataprepEditorViewport::OnFocusViewportToSelection()
 
 void SDataprepEditorViewport::InitializeDefaultMaterials()
 {
+	auto CreateMaterialFunc = [](FString MaterialName) -> UMaterial*
+	{
+		FSoftObjectPath MaterialSoftPath = FSoftObjectPath(FString("/DataprepEditor/") + MaterialName + "." + MaterialName);
+		if(UMaterial* Material = Cast< UMaterial >( MaterialSoftPath.TryLoad() ))
+		{
+			Material->SetFlags(RF_Transient);
+			Material->ClearFlags(RF_Transactional);
+
+			return Material;
+		}
+
+		check(false);
+		return nullptr;
+	};
+
 	const int32 DefaultMaterialsCount = 4;
 
 	TArray< UMaterialInterface* > Materials;
@@ -728,15 +743,7 @@ void SDataprepEditorViewport::InitializeDefaultMaterials()
 
 	if(!PreviewMaterial.IsValid())
 	{
-		PreviewMaterial = TWeakObjectPtr<UMaterial>( Cast< UMaterial >( FSoftObjectPath("/DataprepEditor/PreviewMaterial.PreviewMaterial").TryLoad() ) );
-
-		if(!PreviewMaterial.IsValid())
-		{
-			PreviewMaterial = TWeakObjectPtr<UMaterial>( Cast< UMaterial >( FSoftObjectPath("/Engine/EngineMaterials/DefaultMaterial.DefaultMaterial").TryLoad() ) );
-		}
-
-		check( PreviewMaterial.IsValid() );
-
+		PreviewMaterial = TWeakObjectPtr<UMaterial>( CreateMaterialFunc("PreviewMaterial") );
 		Materials.Add( PreviewMaterial.Get() );
 	}
 
@@ -758,28 +765,19 @@ void SDataprepEditorViewport::InitializeDefaultMaterials()
 
 	if(!XRayMaterial.IsValid())
 	{
-		XRayMaterial = TWeakObjectPtr<UMaterial>( Cast< UMaterial >( FSoftObjectPath("/DataprepEditor/xray_master.xray_master").TryLoad() ) );
-
-		check( XRayMaterial.IsValid() );
-
+		XRayMaterial = TWeakObjectPtr<UMaterial>( CreateMaterialFunc("xray_master") );
 		Materials.Add( XRayMaterial.Get() );
 	}
 
 	if(!BackFaceMaterial.IsValid())
 	{
-		BackFaceMaterial = TWeakObjectPtr<UMaterial>( Cast< UMaterial >( FSoftObjectPath("/DataprepEditor/BackFaceMaterial.BackFaceMaterial").TryLoad() ) );
-
-		check( BackFaceMaterial.IsValid() );
-
+		BackFaceMaterial = TWeakObjectPtr<UMaterial>( CreateMaterialFunc("BackFaceMaterial") );
 		Materials.Add( BackFaceMaterial.Get() );
 	}
 
 	if(!PerMeshMaterial.IsValid())
 	{
-		PerMeshMaterial = TWeakObjectPtr<UMaterial>( Cast< UMaterial >( FSoftObjectPath("/DataprepEditor/PerMeshMaterial.PerMeshMaterial").TryLoad() ) );
-
-		check( PerMeshMaterial.IsValid() );
-
+		PerMeshMaterial = TWeakObjectPtr<UMaterial>( CreateMaterialFunc("PerMeshMaterial") );
 		Materials.Add( PerMeshMaterial.Get() );
 	}
 
@@ -807,10 +805,7 @@ void SDataprepEditorViewport::InitializeDefaultMaterials()
 
 	if(!ReflectionMaterial.IsValid())
 	{
-		ReflectionMaterial = TWeakObjectPtr<UMaterial>( Cast< UMaterial >( FSoftObjectPath("/DataprepEditor/ReflectionMaterial.ReflectionMaterial").TryLoad() ) );
-
-		check( ReflectionMaterial.IsValid() );
-
+		ReflectionMaterial = TWeakObjectPtr<UMaterial>( CreateMaterialFunc("ReflectionMaterial") );
 		Materials.Add( ReflectionMaterial.Get() );
 	}
 
@@ -1626,7 +1621,7 @@ namespace DataprepEditor3DPreviewUtils
 	}
 
 	// Copied from DatasmithImporterImpl::CompileMaterial
-	void CompileMaterials(const TArray< UMaterialInterface* > Materials)
+	void CompileMaterials(const TArray< UMaterialInterface* >& Materials)
 	{
 		if(Materials.Num() > 0)
 		{
