Reverted change to DataprepCore.Build.cs as it is not necessary

diff --git a/Engine/Plugins/Enterprise/DataprepEditor/Source/DataprepCore/DataprepCore.Build.cs b/Engine/Plugins/Enterprise/DataprepEditor/Source/DataprepCore/DataprepCore.Build.cs
--- a/Engine/Plugins/Enterprise/DataprepEditor/Source/DataprepCore/DataprepCore.Build.cs
+++ b/Engine/Plugins/Enterprise/DataprepEditor/Source/DataprepCore/DataprepCore.Build.cs
@@ -25,7 +25,6 @@ namespace UnrealBuildTool.Rules
 					"MessageLog",
 					"PropertyEditor",
 					"RenderCore",
-					"RHI",
 					"Slate",
 					"SlateCore",
 					"StaticMeshDescription",
