The editor dataprep asset library now work directly on the new dataprep data structure.

diff --git a/Engine/Plugins/Enterprise/DataprepEditor/Source/DataprepCore/Public/DataprepActionAsset.h b/Engine/Plugins/Enterprise/DataprepEditor/Source/DataprepCore/Public/DataprepActionAsset.h
--- a/Engine/Plugins/Enterprise/DataprepEditor/Source/DataprepCore/Public/DataprepActionAsset.h
+++ b/Engine/Plugins/Enterprise/DataprepEditor/Source/DataprepCore/Public/DataprepActionAsset.h
@@ -455,7 +455,7 @@ private:
 	bool bWorkingSetHasChanged;
 
 	/** UI label of the action */
-	UPROPERTY()
+	UPROPERTY(EditAnywhere, Category = "Label")
 	FString Label;
 
 	/** Package which static meshes will be added to */
diff --git a/Engine/Plugins/Enterprise/DataprepEditor/Source/DataprepCore/Public/DataprepAsset.h b/Engine/Plugins/Enterprise/DataprepEditor/Source/DataprepCore/Public/DataprepAsset.h
--- a/Engine/Plugins/Enterprise/DataprepEditor/Source/DataprepCore/Public/DataprepAsset.h
+++ b/Engine/Plugins/Enterprise/DataprepEditor/Source/DataprepCore/Public/DataprepAsset.h
@@ -75,7 +75,7 @@ public:
 	 * @return The index of the added action or index none if the action is invalid
 	 * @remark If action is nullptr, a new DataprepActionAsset is simply created
 	 */
-	int32 AddAction(const UDataprepActionAsset* Action);
+	int32 AddAction(const UDataprepActionAsset* Action = nullptr);
 
 	/**
 	 * Add the actions to the Dataprep asset
diff --git a/Engine/Plugins/Enterprise/DataprepEditor/Source/DataprepEditor/Private/DataprepEditorModule.cpp b/Engine/Plugins/Enterprise/DataprepEditor/Source/DataprepEditor/Private/DataprepEditorModule.cpp
--- a/Engine/Plugins/Enterprise/DataprepEditor/Source/DataprepEditor/Private/DataprepEditorModule.cpp
+++ b/Engine/Plugins/Enterprise/DataprepEditor/Source/DataprepEditor/Private/DataprepEditorModule.cpp
@@ -129,24 +129,6 @@ public:
 		return SNullWidget::NullWidget;
 	}
 
-	virtual UEdGraphNode* CreateDataprepActionInGraph(UEdGraph& Graph, UDataprepActionAsset*& OutDataprepAction) override
-	{
-		FVector2D NullLocation( EForceInit::ForceInitToZero );
-		UK2Node_DataprepAction* DataprepActionNode = DataprepSchemaActionUtils::SpawnEdGraphNode<UK2Node_DataprepAction>( Graph, NullLocation );
-		DataprepActionNode->CreateDataprepActionAsset();
-		OutDataprepAction = DataprepActionNode->GetDataprepAction();
-		return DataprepActionNode;
-	}
-
-	virtual UEdGraphNode* CreateDataprepActionInGraphByDuplication(UEdGraph& Graph, const UDataprepActionAsset& DataprepAssetToBeDuplicated, UDataprepActionAsset*& OutDataprepAction)
-	{
-		FVector2D NullLocation( EForceInit::ForceInitToZero );
-		UK2Node_DataprepAction* DataprepActionNode = DataprepSchemaActionUtils::SpawnEdGraphNode<UK2Node_DataprepAction>(Graph, NullLocation);
-		DataprepActionNode->DuplicateExistingDataprepActionAsset( DataprepAssetToBeDuplicated );
-		OutDataprepAction = DataprepActionNode->GetDataprepAction();
-		return DataprepActionNode;
-	}
-
 private:
 	TSharedPtr<FExtensibilityManager> MenuExtensibilityManager;
 	TSharedPtr<FExtensibilityManager> ToolBarExtensibilityManager;
diff --git a/Engine/Plugins/Enterprise/DataprepEditor/Source/DataprepEditor/Public/DataprepEditorModule.h b/Engine/Plugins/Enterprise/DataprepEditor/Source/DataprepEditor/Public/DataprepEditorModule.h
--- a/Engine/Plugins/Enterprise/DataprepEditor/Source/DataprepEditor/Public/DataprepEditorModule.h
+++ b/Engine/Plugins/Enterprise/DataprepEditor/Source/DataprepEditor/Public/DataprepEditorModule.h
@@ -47,9 +47,4 @@ public:
 	virtual TSharedRef<SWidget> CreateDataprepDetailsView(UObject* ObjectToDetail) = 0;
 
 	static EAssetTypeCategories::Type DataprepCategoryBit;
-
-	// Temporary remove this when new graph is implemented (to create dataprep action node
-	virtual class UEdGraphNode* CreateDataprepActionInGraph(class UEdGraph& Graph, class UDataprepActionAsset*& OutDataprepAction) = 0;
-
-	virtual UEdGraphNode* CreateDataprepActionInGraphByDuplication(UEdGraph& Graph, const UDataprepActionAsset& DataprepAssetToBeDuplicated, UDataprepActionAsset*& OutDataprepAction) = 0;
 };
diff --git a/Engine/Plugins/Enterprise/DataprepEditor/Source/DataprepEditorScriptingUtilities/Private/EditorDataprepAssetLibrary.cpp b/Engine/Plugins/Enterprise/DataprepEditor/Source/DataprepEditorScriptingUtilities/Private/EditorDataprepAssetLibrary.cpp
--- a/Engine/Plugins/Enterprise/DataprepEditor/Source/DataprepEditorScriptingUtilities/Private/EditorDataprepAssetLibrary.cpp
+++ b/Engine/Plugins/Enterprise/DataprepEditor/Source/DataprepEditorScriptingUtilities/Private/EditorDataprepAssetLibrary.cpp
@@ -209,14 +209,7 @@ UDataprepActionAsset* UEditorDataprepAssetLibrary::AddAction(UDataprepAsset* Dat
 
 	if ( DataprepAsset )
 	{
-
-		UBlueprint* DataprepRecipeBP = DataprepAsset->GetRecipeBP();
-		check( DataprepRecipeBP->UbergraphPages.Num() == 1 );
-		UEdGraph* DataprepGraph = DataprepRecipeBP->UbergraphPages[0];
-
-		UEdGraphNode* NewActionNode = IDataprepEditorModule::Get().CreateDataprepActionInGraph( *DataprepGraph, NewDataprepAction );
-
-		DataprepAsset->AddActionUsingBP( NewActionNode );
+		NewDataprepAction = DataprepAsset->GetAction( DataprepAsset->AddAction() );
 	}
 	else
 	{
@@ -240,12 +233,8 @@ UDataprepActionAsset* UEditorDataprepAssetLibrary::AddActionByDuplication(UDatap
 	}
 	else
 	{
-		UBlueprint* DataprepRecipeBP = DataprepAsset->GetRecipeBP();
-		check(DataprepRecipeBP->UbergraphPages.Num() == 1);
-		UEdGraph* DataprepGraph = DataprepRecipeBP->UbergraphPages[0];
-
-		UEdGraphNode* NewActionNode = IDataprepEditorModule::Get().CreateDataprepActionInGraphByDuplication(*DataprepGraph, *ActionToDuplicate, NewDataprepAction);
-		DataprepAsset->AddActionUsingBP( NewActionNode );
+		uint32 NewActionIndex = DataprepAsset->AddAction( ActionToDuplicate );
+		NewDataprepAction = DataprepAsset->GetAction( NewActionIndex );
 	}
 
 	return NewDataprepAction;
