Revisited UI of action and step nodes based on designer's feedback

Note: There is another submission to complete the implementation of the UI changes recoomended by Jan

diff --git a/Engine/Plugins/Enterprise/DataprepEditor/Source/DataprepEditor/Private/DataprepEditorStyle.cpp b/Engine/Plugins/Enterprise/DataprepEditor/Source/DataprepEditor/Private/DataprepEditorStyle.cpp
--- a/Engine/Plugins/Enterprise/DataprepEditor/Source/DataprepEditor/Private/DataprepEditorStyle.cpp
+++ b/Engine/Plugins/Enterprise/DataprepEditor/Source/DataprepEditor/Private/DataprepEditorStyle.cpp
@@ -31,6 +31,7 @@ void FDataprepEditorStyle::Initialize()
 
 	const FVector2D Icon20x16(16.0f, 16.0f);
 	const FVector2D Icon20x20(20.0f, 20.0f);
+	const FVector2D Icon32x32(32.0f, 32.0f);
 	const FVector2D Icon40x40(40.0f, 40.0f);
 
 	StyleSet->Set("DataprepEditor.Producer", new IMAGE_PLUGIN_BRUSH("Icons/Producer24", Icon20x20));
@@ -59,27 +60,68 @@ void FDataprepEditorStyle::Initialize()
 	StyleSet->Set("DataprepEditor.ExecutePipeline.Selected", new IMAGE_PLUGIN_BRUSH("Icons/ExecutePipeline", Icon40x40));
 	StyleSet->Set("DataprepEditor.ExecutePipeline.Selected.Small", new IMAGE_PLUGIN_BRUSH("Icons/ExecutePipeline", Icon20x20));
 
+	StyleSet->Set("DataprepEditor.TrackNode.Slot", new IMAGE_PLUGIN_BRUSH("CircleBox", Icon32x32));
+
 	StyleSet->Set( "DataprepEditor.SectionFont", DEFAULT_FONT( "Bold", 10 ) );
 
 	// Dataprep action UI
 	{
-		StyleSet->Set("DataprepAction.Padding", 2.f);
+		StyleSet->Set("Dataprep.TextSeparator.Color", FLinearColor(FColor(200, 200, 200, 200)) );
+
+		const FMargin ActionOuter(1.f, 2.f, 2.f, 2.f);
+		StyleSet->Set("DataprepAction.Body.Padding", ActionOuter + FMargin(5.f, 0.f, 0.f, 0.f));
+		StyleSet->Set("DataprepAction.Outter.Regular.Padding", ActionOuter);
+		StyleSet->Set("DataprepAction.Outter.Selected.Padding", FMargin(0.f));
+		StyleSet->Set("DataprepAction.Steps.Padding", ActionOuter + FMargin(10.f, 0.f, 4.f, 0.f));
+
+		StyleSet->Set("DataprepAction.OutlineColor", FLinearColor(FColor(10, 177, 51)));
+		StyleSet->Set("DataprepAction.BackgroundColor", FLinearColor(FColor(61, 61, 61)));
+
+		{
+			FTextBlockStyle TilteTextStyle = FEditorStyle::GetWidgetStyle< FTextBlockStyle >("NormalText");
+			TilteTextStyle.SetFont(FCoreStyle::GetDefaultFontStyle("Bold", 11));
+			StyleSet->Set("DataprepAction.TitleTextStyle", TilteTextStyle);
+
+			FEditableTextBoxStyle TitleEditableText = FEditorStyle::GetWidgetStyle< FEditableTextBoxStyle >("ViewportMenu.EditableText");
+			TitleEditableText.SetFont(FCoreStyle::GetDefaultFontStyle("Bold", 11));
+
+			StyleSet->Set( "DataprepAction.TitleInlineEditableText", FInlineEditableTextBlockStyle()
+				.SetTextStyle(TilteTextStyle)
+				.SetEditableTextBoxStyle(TitleEditableText)
+			);
+		}
+
+		StyleSet->Set("DataprepAction.EmptyStep.Background.Hovered", FLinearColor(FColor(66, 66, 66)));
+		StyleSet->Set("DataprepAction.EmptyStep.Background.Normal", FLinearColor(FColor(57, 57, 57)));
+		StyleSet->Set("DataprepAction.EmptyStep.Outer.Hovered", FLinearColor(FColor(117, 117, 117)));
+		StyleSet->Set("DataprepAction.EmptyStep.Outer.Normal", FLinearColor(FColor(85, 85, 85)));
+		StyleSet->Set("DataprepAction.EmptyStep.Text.Hovered", FLinearColor(FColor(230, 230, 230)));
+		StyleSet->Set("DataprepAction.EmptyStep.Text.Normal", FLinearColor(FColor(117, 117, 117)));
 
-		StyleSet->Set("DataprepAction.OutlineColor", FLinearColor(FColor(62, 62, 62)));
+		StyleSet->Set("DataprepAction.EmptyStep.Bottom.Padding", FMargin(0.f, 0.f, 0.f, 5.f));
 
-		StyleSet->Set("DataprepActionStep.BackgroundColor", FLinearColor(FColor(62, 62, 62)) );
+		StyleSet->Set("DataprepActionSteps.BackgroundColor", FLinearColor(FColor(26, 26, 26)));
+
+		StyleSet->Set("DataprepActionStep.BackgroundColor", FLinearColor(FColor(93, 93, 93)) );
 		StyleSet->Set("DataprepActionStep.DragAndDrop", FLinearColor(FColor(212, 212, 59)) );
 		StyleSet->Set("DataprepActionStep.Selected", FLinearColor(FColor(1, 202, 252)) );
-		StyleSet->Set("DataprepActionStep.Filter.OutlineColor", FLinearColor(FColor(67, 105, 124)));
-		StyleSet->Set("DataprepActionStep.Operation.OutlineColor", FLinearColor(FColor(87, 107, 61)) );
-
-		StyleSet->Set("DataprepActionBlock.TitleBackgroundColor", FLinearColor(0.065307f, 0.065307f, 0.065307f));
-		StyleSet->Set("DataprepActionBlock.ContentBackgroundColor", FLinearColor(0.11f, 0.11f, 0.11f));
-		FTextBlockStyle TilteTextBlockStyle = FEditorStyle::GetWidgetStyle< FTextBlockStyle >("NormalText");
-		TilteTextBlockStyle.SetFont(FCoreStyle::GetDefaultFontStyle("Bold", 11));
-		StyleSet->Set("DataprepActionBlock.TitleTextBlockStyle", TilteTextBlockStyle);
-		StyleSet->Set("DataprepActionSteps.BackgroundColor", FLinearColor( 0.1033, 0.1033, 0.1033));
-		StyleSet->Set("DataprepActionStep.Padding", 10.f);
+		StyleSet->Set("DataprepActionStep.Filter.OutlineColor", FLinearColor(FColor(220, 125, 67)));
+		StyleSet->Set("DataprepActionStep.Operation.OutlineColor", FLinearColor(FColor(67, 177, 220)) );
+		StyleSet->Set("DataprepActionStep.Separator.Color", FLinearColor(FColor(182, 219, 192)) );
+
+		StyleSet->Set("DataprepActionStep.Outter.Regular.Padding", FMargin(10.f, 3.f, 10.f, 3.f));
+		StyleSet->Set("DataprepActionStep.Outter.Selected.Padding", FMargin(10.f, 0.f, 4.f, 0.f));
+		StyleSet->Set("DataprepActionStep.Padding", FMargin(15.f, 3.f, 5.f, 3.f));
+
+		StyleSet->Set("DataprepActionStep.DnD.Outter.Padding", FMargin(0.f, 5.f, 0.f, 5.f));
+		StyleSet->Set("DataprepActionStep.DnD.Inner.Padding", FMargin(5.f, 5.f, 5.f, 5.f));
+
+		StyleSet->Set("DataprepActionBlock.ContentBackgroundColor.Old", FLinearColor(0.11f, 0.11f, 0.11f));
+		{
+			FTextBlockStyle TilteTextBlockStyle = FEditorStyle::GetWidgetStyle< FTextBlockStyle >("NormalText");
+			TilteTextBlockStyle.SetFont(FCoreStyle::GetDefaultFontStyle("Bold", 10));
+			StyleSet->Set("DataprepActionBlock.TitleTextBlockStyle", TilteTextBlockStyle);
+		}
 	}
 
 	// DataprepGraphEditor
diff --git a/Engine/Plugins/Enterprise/DataprepEditor/Source/DataprepEditor/Private/DataprepEditorStyle.h b/Engine/Plugins/Enterprise/DataprepEditor/Source/DataprepEditor/Private/DataprepEditorStyle.h
--- a/Engine/Plugins/Enterprise/DataprepEditor/Source/DataprepEditor/Private/DataprepEditorStyle.h
+++ b/Engine/Plugins/Enterprise/DataprepEditor/Source/DataprepEditor/Private/DataprepEditorStyle.h
@@ -35,6 +35,11 @@ public:
 		return StyleSet->GetMargin( PropertyName, Specifier );
 	}
 
+	static const FSlateBrush* GetBrush( FName PropertyName, const ANSICHAR* Specifier = nullptr )
+	{
+		return StyleSet->GetBrush( PropertyName, Specifier );
+	}
+
 	template< typename WidgetStyleType >            
 	static const WidgetStyleType& GetWidgetStyle( FName PropertyName, const ANSICHAR* Specifier = nullptr )
 	{
diff --git a/Engine/Plugins/Enterprise/DataprepEditor/Source/DataprepEditor/Private/Widgets/DataprepGraph/SDataprepActionBlock.cpp b/Engine/Plugins/Enterprise/DataprepEditor/Source/DataprepEditor/Private/Widgets/DataprepGraph/SDataprepActionBlock.cpp
--- a/Engine/Plugins/Enterprise/DataprepEditor/Source/DataprepEditor/Private/Widgets/DataprepGraph/SDataprepActionBlock.cpp
+++ b/Engine/Plugins/Enterprise/DataprepEditor/Source/DataprepEditor/Private/Widgets/DataprepGraph/SDataprepActionBlock.cpp
@@ -21,8 +21,10 @@
 #include "Styling/ISlateStyle.h"
 #include "Styling/SlateStyleRegistry.h"
 #include "Widgets/Colors/SColorBlock.h"
+#include "Widgets/Images/SImage.h"
 #include "Widgets/Layout/Anchors.h"
 #include "Widgets/Layout/SConstraintCanvas.h"
+#include "Widgets/Layout/SSeparator.h"
 #include "Widgets/SBoxPanel.h"
 #include "Widgets/SNullWidget.h"
 #include "Widgets/Text/STextBlock.h"
@@ -34,9 +36,15 @@ void SDataprepActionBlock::Construct(const FArguments& InArgs, const TSharedRef<
 	DataprepActionContext = InDataprepActionContext;
 
 	const float DefaultPadding = FDataprepEditorStyle::GetFloat( "DataprepAction.Padding" );
-	const FLinearColor OutlineColor = GetOutlineColor();
+	const FLinearColor OutlineColor = GetOutlineColor().GetSpecifiedColor();
 #ifndef NO_BLUEPRINT
-	bIsSimplifiedGraph = false;
+	bIsSimplifiedGraph = InArgs._IsSimplified;
+
+	if(bIsSimplifiedGraph)
+	{
+		ConstructForSimplified();
+		return;
+	}
 #endif
 
 	ChildSlot
@@ -87,7 +95,7 @@ void SDataprepActionBlock::Construct(const FArguments& InArgs, const TSharedRef<
 			// The content zone of the action block
 			+ SVerticalBox::Slot()
 			.AutoHeight()
-			.Padding( !bIsSimplifiedGraph ? FMargin( DefaultPadding ) : FMargin( 2.f * DefaultPadding, 0.f, 0.f, 0.f ) )
+			.Padding( FMargin( DefaultPadding ) )
 			[
 				SNew( SConstraintCanvas )
 
@@ -112,6 +120,56 @@ void SDataprepActionBlock::Construct(const FArguments& InArgs, const TSharedRef<
 	];
 }
 
+void SDataprepActionBlock::ConstructForSimplified()
+{
+	const FLinearColor OutlineColor = GetOutlineColor().GetSpecifiedColor();
+
+	ChildSlot
+	[
+		SNew( SVerticalBox )
+
+		//The title of the block
+		+ SVerticalBox::Slot()
+		.AutoHeight()
+		.Padding(0.f)
+		[
+			SNew( SHorizontalBox )
+
+			+ SHorizontalBox::Slot()
+			.AutoWidth()
+			.Padding( FMargin( 5.f, 5.f, 5.f, 2.f ) )
+			[
+				GetTitleWidget()
+			]
+		]
+
+		+ SVerticalBox::Slot()
+		.AutoHeight()
+		.Padding(5.f, 2.f)
+		[
+			SNew( SSeparator )
+			.SeparatorImage(FEditorStyle::GetBrush( "ThinLine.Horizontal" ))
+			.Thickness(1.f)
+			.Orientation(EOrientation::Orient_Horizontal)
+			.ColorAndOpacity(FDataprepEditorStyle::GetColor("Dataprep.TextSeparator.Color"))
+		]
+
+		//The content of the block
+		+ SVerticalBox::Slot()
+		.AutoHeight()
+		.Padding(0.f)
+		[
+			SNew( SHorizontalBox )
+
+			+ SHorizontalBox::Slot()
+			.Padding( FMargin( 5.f, 0.f, 5.f, 0.f ) )
+			[
+				GetContentWidget()
+			]
+		]
+	];
+}
+
 TSharedRef<SWidget> SDataprepActionBlock::GetBlockTitleWidget()
 {
 	const float DefaultPadding = FDataprepEditorStyle::GetFloat( "DataprepAction.Padding" );
@@ -341,7 +399,7 @@ FReply SDataprepActionBlock::OnDragDetected(const FGeometry& MyGeometry, const F
 	return FReply::Handled().BeginDragDrop( MoveTemp( DragDropOperation ) );
 }
 
-FLinearColor SDataprepActionBlock::GetOutlineColor() const
+FSlateColor SDataprepActionBlock::GetOutlineColor() const
 {
 	return FDataprepEditorStyle::GetColor( "DataprepAction.OutlineColor" );
 }
@@ -365,8 +423,7 @@ TSharedRef<SWidget> SDataprepActionBlock::GetTitleWidget()
 
 TSharedRef<SWidget> SDataprepActionBlock::GetTitleBackgroundWidget()
 {
-	return SNew( SColorBlock )
-		.Color( GetOutlineColor() );
+	return SNew( SColorBlock ).Color( GetOutlineColor().GetSpecifiedColor() );
 }
 
 TSharedRef<SWidget> SDataprepActionBlock::GetContentWidget()
@@ -377,7 +434,7 @@ TSharedRef<SWidget> SDataprepActionBlock::GetContentWidget()
 TSharedRef<SWidget> SDataprepActionBlock::GetContentBackgroundWidget()
 {
 	return SNew( SColorBlock )
-		.Color( FDataprepEditorStyle::GetColor("DataprepActionBlock.ContentBackgroundColor") );
+		.Color( FDataprepEditorStyle::GetColor("DataprepActionBlock.ContentBackgroundColor.Old") );
 }
 
 void SDataprepActionBlock::PopulateMenuBuilder(FMenuBuilder& MenuBuilder)
diff --git a/Engine/Plugins/Enterprise/DataprepEditor/Source/DataprepEditor/Private/Widgets/DataprepGraph/SDataprepActionBlock.h b/Engine/Plugins/Enterprise/DataprepEditor/Source/DataprepEditor/Private/Widgets/DataprepGraph/SDataprepActionBlock.h
--- a/Engine/Plugins/Enterprise/DataprepEditor/Source/DataprepEditor/Private/Widgets/DataprepGraph/SDataprepActionBlock.h
+++ b/Engine/Plugins/Enterprise/DataprepEditor/Source/DataprepEditor/Private/Widgets/DataprepGraph/SDataprepActionBlock.h
@@ -14,8 +14,16 @@ struct FTextBlockStyle;
  */
 class SDataprepActionBlock : public SCompoundWidget
 {
+#ifndef NO_BLUEPRINT
+	SLATE_BEGIN_ARGS(SDataprepActionBlock)
+	: _IsSimplified(false)
+	{}
+
+		SLATE_ARGUMENT( bool, IsSimplified )
+#else
 	SLATE_BEGIN_ARGS(SDataprepActionBlock) {}
-	SLATE_END_ARGS();
+#endif
+	SLATE_END_ARGS()
 
 	void Construct(const FArguments& InArgs, const TSharedRef<FDataprepSchemaActionContext>& DataprepActionContext);
 
@@ -23,18 +31,28 @@ class SDataprepActionBlock : public SCompoundWidget
 	virtual FVector2D ComputeDesiredSize(float) const override
 	{
 		const FVector2D ChildSize = ChildSlot.GetWidget()->GetDesiredSize();
-		return FVector2D( 280.f, ChildSize.Y );
+		return FVector2D( FMath::Max(280.f, ChildSize.X) , ChildSize.Y );
 	}
 	virtual FReply OnMouseButtonDown(const FGeometry& MyGeometry, const FPointerEvent& MouseEvent) override;
 	virtual FReply OnMouseButtonUp(const FGeometry& MyGeometry, const FPointerEvent& MouseEvent) override;
 	virtual FReply OnDragDetected(const FGeometry& MyGeometry, const FPointerEvent& MouseEvent) override;
 	// End of SWidget Interface
 
+	/**
+	 * Return the outline color
+	 */
+	virtual FSlateColor GetOutlineColor() const;
+
 public:
 
 	/** Return the block title */
 	TSharedRef<SWidget> GetBlockTitleWidget();
 
+	/**
+	 * Return the title widget
+	 */
+	virtual TSharedRef<SWidget> GetTitleWidget();
+
 #ifndef NO_BLUEPRINT
 	// Temporary boolean to indicate if the widget has been instantiated in a regular BP graph or the simplified one
 	bool bIsSimplifiedGraph;
@@ -42,10 +60,7 @@ public:
 
 protected:
 
-	/**
-	* Return the outline color
-	*/
-	virtual FLinearColor GetOutlineColor() const;
+	void ConstructForSimplified();
 
 	/**
 	 * Return the block title.
@@ -53,15 +68,10 @@ protected:
 	 */
 	virtual FText GetBlockTitle() const;
 
-	/**
-	 * Return the title widget
-	 */
- 	virtual TSharedRef<SWidget> GetTitleWidget();
-
 	/**
 	 * Get title background widget
 	 */
-	virtual TSharedRef<SWidget> GetTitleBackgroundWidget();
+	TSharedRef<SWidget> GetTitleBackgroundWidget();
 
 	/**
 	 * Return the content widget
@@ -71,7 +81,7 @@ protected:
 	/**
 	 * Get content background widget
 	 */
-	virtual TSharedRef<SWidget> GetContentBackgroundWidget();
+	TSharedRef<SWidget> GetContentBackgroundWidget();
 
 	/**
 	 * Populate the right click menu by overriding this function.
diff --git a/Engine/Plugins/Enterprise/DataprepEditor/Source/DataprepEditor/Private/Widgets/DataprepGraph/SDataprepFilter.cpp b/Engine/Plugins/Enterprise/DataprepEditor/Source/DataprepEditor/Private/Widgets/DataprepGraph/SDataprepFilter.cpp
--- a/Engine/Plugins/Enterprise/DataprepEditor/Source/DataprepEditor/Private/Widgets/DataprepGraph/SDataprepFilter.cpp
+++ b/Engine/Plugins/Enterprise/DataprepEditor/Source/DataprepEditor/Private/Widgets/DataprepGraph/SDataprepFilter.cpp
@@ -30,7 +30,7 @@ void SDataprepFilter::Construct(const FArguments& InArgs, UDataprepFilter& InFil
 	TAttribute<FText> TooltipTextAttribute = MakeAttributeSP( this, &SDataprepFilter::GetTooltipText );
 	SetToolTipText( TooltipTextAttribute );
 
-	SDataprepActionBlock::Construct( SDataprepActionBlock::FArguments(), InDataprepActionContext );
+	SDataprepActionBlock::Construct( SDataprepActionBlock::FArguments().IsSimplified(InArgs._IsSimplified), InDataprepActionContext );
 }
 
 void SDataprepFilter::Tick(const FGeometry& AllottedGeometry, const double InCurrentTime, const float InDeltaTime)
@@ -44,9 +44,10 @@ void SDataprepFilter::Tick(const FGeometry& AllottedGeometry, const double InCur
 	}
 }
 
-FLinearColor SDataprepFilter::GetOutlineColor() const
+FSlateColor SDataprepFilter::GetOutlineColor() const
 {
-	return FDataprepEditorStyle::GetColor( "DataprepActionStep.Filter.OutlineColor" );
+	static const FSlateColor OutlineColor = FDataprepEditorStyle::GetColor( "DataprepActionStep.Filter.OutlineColor" );
+	return OutlineColor;
 }
 
 FText SDataprepFilter::GetBlockTitle() const
diff --git a/Engine/Plugins/Enterprise/DataprepEditor/Source/DataprepEditor/Private/Widgets/DataprepGraph/SDataprepFilter.h b/Engine/Plugins/Enterprise/DataprepEditor/Source/DataprepEditor/Private/Widgets/DataprepGraph/SDataprepFilter.h
--- a/Engine/Plugins/Enterprise/DataprepEditor/Source/DataprepEditor/Private/Widgets/DataprepGraph/SDataprepFilter.h
+++ b/Engine/Plugins/Enterprise/DataprepEditor/Source/DataprepEditor/Private/Widgets/DataprepGraph/SDataprepFilter.h
@@ -16,8 +16,15 @@ struct FDataprepSchemaActionContext;
 class SDataprepFilter : public SDataprepActionBlock, public FGCObject
 {
 public:
+#ifndef NO_BLUEPRINT
+	SLATE_BEGIN_ARGS(SDataprepFilter)
+	: _IsSimplified(false)
+	{}
+		SLATE_ARGUMENT( bool, IsSimplified )
+#else
 	SLATE_BEGIN_ARGS(SDataprepFilter) {}
-	SLATE_END_ARGS();
+#endif
+	SLATE_END_ARGS()
 
 	void Construct(const FArguments& InArgs, UDataprepFilter& InFilter, const TSharedRef<FDataprepSchemaActionContext>& InDataprepActionContext);
 
@@ -26,7 +33,7 @@ public:
 protected:
 
 	// SDataprepActionBlock interface
-	virtual FLinearColor GetOutlineColor() const override;
+	virtual FSlateColor GetOutlineColor() const override;
 	virtual FText GetBlockTitle() const override;
 	virtual TSharedRef<SWidget> GetTitleWidget() override;
 	virtual TSharedRef<SWidget> GetContentWidget() override;
diff --git a/Engine/Plugins/Enterprise/DataprepEditor/Source/DataprepEditor/Private/Widgets/DataprepGraph/SDataprepGraphActionNode.cpp b/Engine/Plugins/Enterprise/DataprepEditor/Source/DataprepEditor/Private/Widgets/DataprepGraph/SDataprepGraphActionNode.cpp
--- a/Engine/Plugins/Enterprise/DataprepEditor/Source/DataprepEditor/Private/Widgets/DataprepGraph/SDataprepGraphActionNode.cpp
+++ b/Engine/Plugins/Enterprise/DataprepEditor/Source/DataprepEditor/Private/Widgets/DataprepGraph/SDataprepGraphActionNode.cpp
@@ -16,16 +16,21 @@
 #include "NodeFactory.h"
 #include "SGraphPanel.h"
 #include "Widgets/Colors/SColorBlock.h"
+#include "Widgets/Images/SImage.h"
+#include "Widgets/Layout/SBox.h"
+#include "Widgets/Layout/SSeparator.h"
 #include "Widgets/SBoxPanel.h"
 #include "Widgets/SOverlay.h"
-#include "Widgets/Layout/SBox.h"
+#include "Widgets/Text/SInlineEditableTextBlock.h"
+#include "Widgets/Text/STextBlock.h"
 
 #define LOCTEXT_NAMESPACE "DataprepGraphEditor"
 
 class SDataprepGraphActionProxyNode : public SGraphNode
 {
 public:
-	SLATE_BEGIN_ARGS(SDataprepGraphActionStepNode) {}
+	SLATE_BEGIN_ARGS(SDataprepGraphActionProxyNode) {}
+		SLATE_ARGUMENT( TSharedPtr<SInlineEditableTextBlock>, InlineEditableText )
 	SLATE_END_ARGS()
 
 	void Construct(const FArguments& InArgs, const TSharedRef<SDataprepGraphActionNode>& InParentNode)
@@ -35,6 +40,8 @@ public:
 
 		SetCursor(EMouseCursor::Default);
 		UpdateGraphNode();
+
+		InlineEditableText = InArgs._InlineEditableText;
 	}
 
 	// SWidget interface
@@ -60,6 +67,20 @@ public:
 		];
 	}
 
+	virtual FSlateRect GetTitleRect() const override
+	{
+		const FVector2D NodePosition = GetPosition();
+		const FVector2D NodeSize = InlineEditableText.IsValid() ? InlineEditableText->GetDesiredSize() : GetDesiredSize();
+
+		return FSlateRect( NodePosition.X, NodePosition.Y + NodeSize.Y, NodePosition.X + NodeSize.X, NodePosition.Y );
+	}
+
+	/* Helper function to check if node can be renamed */
+	virtual bool IsNameReadOnly () const override
+	{
+		return ParentNodePtr.IsValid() ? ParentNodePtr.Pin()->IsNameReadOnly() : true;
+	}
+
 	const FSlateBrush* GetShadowBrush(bool bSelected) const
 	{
 		return  FEditorStyle::GetNoBrush();
@@ -102,8 +123,16 @@ private:
  * The SDataprepEmptyActionStepNode is a helper class that handles drag and drop event at
  * the bottom of the SDataprepGraphActionNode widget
  */
-class SDataprepEmptyActionStepNode : public SCompoundWidget
+class SDataprepEmptyActionStepNode : public SVerticalBox
 {
+	enum EColorType : uint8
+	{
+		TextColor = 0,
+		OuterColor,
+		InnerColor,
+		ColorMaxType
+	};
+
 public:
 	SLATE_BEGIN_ARGS(SDataprepEmptyActionStepNode) {}
 	SLATE_END_ARGS();
@@ -116,20 +145,54 @@ public:
 
 		bIsHovered = false;
 
-		ChildSlot
+		SVerticalBox::Construct(SVerticalBox::FArguments());
+
+		TAttribute<FSlateColor> TextColorAndOpacity = TAttribute<FSlateColor>::Create(TAttribute<FSlateColor>::FGetter::CreateSP(this, &SDataprepEmptyActionStepNode::GetColor, EColorType::TextColor));
+		TAttribute<FSlateColor> OuterColorAndOpacity = TAttribute<FSlateColor>::Create(TAttribute<FSlateColor>::FGetter::CreateSP(this, &SDataprepEmptyActionStepNode::GetColor, EColorType::OuterColor));
+		TAttribute<FSlateColor> InnerColorAndOpacity = TAttribute<FSlateColor>::Create(TAttribute<FSlateColor>::FGetter::CreateSP(this, &SDataprepEmptyActionStepNode::GetColor, EColorType::InnerColor));
+
+		AddSlot()
+		.AutoHeight()
+		.Padding(FDataprepEditorStyle::GetMargin( "DataprepActionStep.Padding" ))
 		[
-			SNew(SVerticalBox)
+			SNew(SOverlay)
 
-			+SVerticalBox::Slot()
-			.AutoHeight()
-			.Padding(10.f, 0.f, 10.f, 0.f)
+			+ SOverlay::Slot()
+			.Padding(0.f)
+			.VAlign(VAlign_Fill)
+			.HAlign(HAlign_Fill)
+			[
+				SNew(SImage)
+				.ColorAndOpacity(MoveTemp(OuterColorAndOpacity))
+				.Image(FEditorStyle::GetBrush( "Graph.StateNode.Body" ))
+			]
+
+			+ SOverlay::Slot()
+			.Padding(1.f)
+			.VAlign(VAlign_Fill)
+			.HAlign(HAlign_Fill)
+			[
+				SNew(SImage)
+				.ColorAndOpacity(MoveTemp(InnerColorAndOpacity))
+				.Image(FEditorStyle::GetBrush( "Graph.StateNode.Body" ))
+			]
+
+			+ SOverlay::Slot()
+			.Padding(10.f)
+			.VAlign(VAlign_Fill)
+			.HAlign(HAlign_Fill)
 			[
-				SNew(SBorder)
-				.BorderBackgroundColor( this, &SDataprepEmptyActionStepNode::GetBorderBackgroundColor )
-				.BorderImage(FEditorStyle::GetBrush("BTEditor.Graph.BTNode.Body"))
+				SNew(SHorizontalBox)
+
+				+ SHorizontalBox::Slot()
+				.VAlign(VAlign_Fill)
+				.HAlign(HAlign_Fill)
 				[
-					SNew(SBox)
-					.HeightOverride(InterStepSpacing)
+					SNew(STextBlock)
+					.Text(LOCTEXT("DataprepEmptyActionStepLabel", "+ Add Step"))
+					.TextStyle( &FDataprepEditorStyle::GetWidgetStyle<FTextBlockStyle>( "DataprepActionBlock.TitleTextBlockStyle" ) )
+					.ColorAndOpacity(MoveTemp(TextColorAndOpacity))
+					.Justification(ETextJustify::Center)
 				]
 			]
 		];
@@ -147,7 +210,7 @@ public:
 			ParentPtr.Pin()->SetHoveredIndex( ParentPtr.Pin()->GetDataprepAction()->GetStepsCount() );
 		}
 
-		SCompoundWidget::OnDragEnter(MyGeometry, DragDropEvent);
+		SVerticalBox::OnDragEnter(MyGeometry, DragDropEvent);
 	}
 
 	virtual FReply OnDragOver(const FGeometry& MyGeometry, const FDragDropEvent& DragDropEvent) override
@@ -161,7 +224,7 @@ public:
 			return FReply::Handled();
 		}
 
-		return SCompoundWidget::OnDragOver(MyGeometry, DragDropEvent);
+		return SVerticalBox::OnDragOver(MyGeometry, DragDropEvent);
 	}
 
 	virtual void OnDragLeave(const FDragDropEvent& DragDropEvent) override
@@ -174,7 +237,7 @@ public:
 
 		ParentPtr.Pin()->SetHoveredIndex( INDEX_NONE );
 
-		SCompoundWidget::OnDragLeave(DragDropEvent);
+		SVerticalBox::OnDragLeave(DragDropEvent);
 	}
 
 	virtual FReply OnDrop(const FGeometry& MyGeometry, const FDragDropEvent& DragDropEvent) override
@@ -190,7 +253,7 @@ public:
 			return DragActionStepNodeOp->DroppedOnNode(DragDropEvent.GetScreenSpacePosition(), NodeAddPosition);
 		}
 
-		return SCompoundWidget::OnDrop(MyGeometry, DragDropEvent);
+		return SVerticalBox::OnDrop(MyGeometry, DragDropEvent);
 	}
 	// End of SWidget Interface
 
@@ -200,9 +263,22 @@ public:
 	}
 
 private:
-	FSlateColor GetBorderBackgroundColor() const
+	FSlateColor GetColor(EColorType Type) const
 	{
-		return ParentPtr.Pin()->GetInsertColor( ParentPtr.Pin()->GetDataprepAction()->GetStepsCount() );
+		static const FLinearColor HoveredColors[EColorType::ColorMaxType] = {
+			FDataprepEditorStyle::GetColor("DataprepAction.EmptyStep.Text.Hovered"),
+			FDataprepEditorStyle::GetColor("DataprepAction.EmptyStep.Outer.Hovered"),
+			FDataprepEditorStyle::GetColor("DataprepAction.EmptyStep.Background.Hovered"),
+		};
+		static const FLinearColor NormalColors[EColorType::ColorMaxType] = {
+			FDataprepEditorStyle::GetColor("DataprepAction.EmptyStep.Text.Normal"),
+			FDataprepEditorStyle::GetColor("DataprepAction.EmptyStep.Outer.Normal"),
+			FDataprepEditorStyle::GetColor("DataprepAction.EmptyStep.Background.Normal"),
+		};
+
+		const int32 EmptyIndex = ParentPtr.Pin()->GetDataprepAction()->GetStepsCount();
+
+		return EmptyIndex == ParentPtr.Pin()->GetHoveredIndex() ? HoveredColors[Type] : NormalColors[Type];
 	}
 
 private:
@@ -223,26 +299,30 @@ void SDataprepGraphActionNode::Construct(const FArguments& InArgs, UDataprepGrap
 
 	GraphNode = InActionNode;
 
-	ProxyNodePtr = SNew(SDataprepGraphActionProxyNode, SharedThis(this));
-
 	SetCursor(EMouseCursor::ResizeLeftRight);
 	UpdateGraphNode();
+
+	ProxyNodePtr = SNew(SDataprepGraphActionProxyNode, SharedThis(this))
+	.InlineEditableText(InlineEditableText);
 }
 
 void SDataprepGraphActionNode::SetParentTrackNode(TSharedPtr<SDataprepGraphTrackNode> InParentTrackNode)
 {
 	ParentTrackNodePtr = InParentTrackNode;
 
-	// Update parent track on step widgets
-	for(TSharedPtr<SDataprepGraphActionStepNode>& ActionStepGraphNode : ActionStepGraphNodes)
+	if(ActionStepListWidgetPtr.IsValid())
 	{
-		ActionStepGraphNode->SetParentTrackNode(InParentTrackNode);
-	}
+		// Update parent track on step widgets
+		for(TSharedPtr<SDataprepGraphActionStepNode>& ActionStepGraphNode : ActionStepGraphNodes)
+		{
+			ActionStepGraphNode->SetParentTrackNode(InParentTrackNode);
+		}
 
-	// Update parent on empty bottom widget
-	FChildren* StepListChildren = ActionStepListWidgetPtr->GetChildren();
-	TSharedRef<SDataprepEmptyActionStepNode> EmptyWidgetPtr = StaticCastSharedRef<SDataprepEmptyActionStepNode>(StepListChildren->GetChildAt(StepListChildren->Num() - 1));
-	EmptyWidgetPtr->SetParentTrackNode(InParentTrackNode);
+		// Update parent on empty bottom widget
+		FChildren* StepListChildren = ActionStepListWidgetPtr->GetChildren();
+		TSharedRef<SDataprepEmptyActionStepNode> EmptyWidgetPtr = StaticCastSharedRef<SDataprepEmptyActionStepNode>(StepListChildren->GetChildAt(StepListChildren->Num() - 1));
+		EmptyWidgetPtr->SetParentTrackNode(InParentTrackNode);
+	}
 }
 
 void SDataprepGraphActionNode::UpdateExecutionOrder()
@@ -253,7 +333,158 @@ void SDataprepGraphActionNode::UpdateExecutionOrder()
 
 void SDataprepGraphActionNode::UpdateProxyNode(const FVector2D& Position)
 {
-	ProxyNodePtr->SetPosition(Position);
+	if(ProxyNodePtr.IsValid())
+	{
+		ProxyNodePtr->SetPosition(Position);
+	}
+}
+
+void SDataprepGraphActionNode::UpdateGraphNode()
+{
+	static bool bUseNew = true;
+	if(!bUseNew)
+	{
+		SGraphNode::UpdateGraphNode();
+		return;
+	}
+
+	// Reset SGraphNode members.
+	InputPins.Empty();
+	OutputPins.Empty();
+	RightNodeBox.Reset();
+	LeftNodeBox.Reset();
+
+	this->ContentScale.Bind( this, &SGraphNode::GetContentScale );
+
+	if(!DataprepActionPtr.IsValid())
+	{
+		this->GetOrAddSlot( ENodeZone::Center )
+		.HAlign(HAlign_Fill)
+		.VAlign(VAlign_Center)
+		[
+			SNew(STextBlock)
+			.ColorAndOpacity( FSlateColor( FLinearColor::Red ) )
+			.Text( FText::FromString( TEXT("This node doesn't have a dataprep action!") ) )
+		];
+
+		return;
+	}
+
+	TAttribute<FText> NodeTitle = TAttribute<FText>::Create(TAttribute<FText>::FGetter::CreateUObject(GraphNode, &UEdGraphNode::GetNodeTitle, ENodeTitleType::EditableTitle));
+	SAssignNew(InlineEditableText, SInlineEditableTextBlock)
+		.Style(FDataprepEditorStyle::Get(), "DataprepAction.TitleInlineEditableText")
+		.Text(NodeTitle)
+		.OnVerifyTextChanged(this, &SDataprepGraphActionNode::OnVerifyNameTextChanged)
+		.OnTextCommitted(this, &SDataprepGraphActionNode::OnNameTextCommited)
+		.IsReadOnly(this, &SDataprepGraphActionNode::IsNameReadOnly)
+		.IsSelected(this, &SDataprepGraphActionNode::IsSelectedExclusively);
+
+	PopulateActionStepListWidget();
+
+	TAttribute<FMargin> OuterPadding = TAttribute<FMargin>::Create(TAttribute<FMargin>::FGetter::CreateSP(this, &SDataprepGraphActionNode::GetOuterPadding));
+
+	this->GetOrAddSlot( ENodeZone::Center )
+	.HAlign(HAlign_Fill)
+	.VAlign(VAlign_Center)
+	[
+		SNew(SBox)
+		.WidthOverride(300.f)
+		[
+			SNew(SVerticalBox)
+
+			+SVerticalBox::Slot()
+			.AutoHeight()
+			[
+				SNew(SOverlay)
+
+				+ SOverlay::Slot()
+				.Padding(OuterPadding)
+				.VAlign(VAlign_Fill)
+				.HAlign(HAlign_Fill)
+				[
+					CreateBackground(FDataprepEditorStyle::GetColor( "DataprepAction.OutlineColor" ))
+				]
+
+				+ SOverlay::Slot()
+				.Padding(FDataprepEditorStyle::GetMargin( "DataprepAction.Body.Padding" ))
+				.VAlign(VAlign_Fill)
+				.HAlign(HAlign_Fill)
+				[
+					CreateBackground(FDataprepEditorStyle::GetColor( "DataprepAction.BackgroundColor" ))
+				]
+
+				+ SOverlay::Slot()
+				.Padding( FDataprepEditorStyle::GetMargin( "DataprepAction.Steps.Padding" ) )
+				.VAlign(VAlign_Top)
+				.HAlign(HAlign_Right)
+				[
+					SNew( SVerticalBox )
+
+					+ SVerticalBox::Slot()
+					.AutoHeight()
+					[
+						SNew( SHorizontalBox )
+
+						+ SHorizontalBox::Slot()
+						.AutoWidth()
+						.Padding(5.f, 5.f, 5.f, 2.f)
+						.HAlign(EHorizontalAlignment::HAlign_Center)
+						[
+							InlineEditableText.ToSharedRef()
+						]
+					]
+
+					+ SVerticalBox::Slot()
+					.AutoHeight()
+					.Padding(5.f, 2.f)
+					[
+						SNew( SSeparator )
+						.SeparatorImage(FEditorStyle::GetBrush( "ThinLine.Horizontal" ))
+						.Thickness(1.f)
+						.Orientation(EOrientation::Orient_Horizontal)
+						.ColorAndOpacity(FDataprepEditorStyle::GetColor("Dataprep.TextSeparator.Color"))
+					]
+
+					//The content of the action
+					+ SVerticalBox::Slot()
+					.AutoHeight()
+					[
+						SNew( SHorizontalBox )
+
+						+ SHorizontalBox::Slot()
+						[
+							ActionStepListWidgetPtr.ToSharedRef()
+						]
+					]
+				]
+			]
+		]
+	];
+}
+
+TSharedRef<SWidget> SDataprepGraphActionNode::CreateBackground(const TAttribute<FSlateColor>& ColorAndOpacity)
+{
+	return SNew(SOverlay)
+
+		+SOverlay::Slot()
+		.VAlign(VAlign_Fill)
+		.HAlign(HAlign_Fill)
+		.Padding(1.f)
+		[
+			SNew(SImage)
+			.ColorAndOpacity(FLinearColor::White)
+			.Image(FEditorStyle::GetBrush( "Graph.StateNode.Body" ))
+		]
+
+		+SOverlay::Slot()
+		.VAlign(VAlign_Fill)
+		.HAlign(HAlign_Fill)
+		.Padding(0.f)
+		[
+			SNew(SImage)
+			.ColorAndOpacity(ColorAndOpacity)
+			.Image(FEditorStyle::GetBrush( "Graph.StateNode.Body" ))
+		];
 }
 
 TSharedRef<SWidget> SDataprepGraphActionNode::CreateNodeContentArea()
@@ -283,23 +514,23 @@ const FSlateBrush* SDataprepGraphActionNode::GetShadowBrush(bool bSelected) cons
 int32 SDataprepGraphActionNode::OnPaint(const FPaintArgs& Args, const FGeometry& AllottedGeometry, const FSlateRect& MyCullingRect, FSlateWindowElementList& OutDrawElements, int32 LayerId, const FWidgetStyle& InWidgetStyle, bool bParentEnabled) const
 {
 	// Since only a proxy is in the graph panel, draw selection outline if applicable 
-	if(SGraphPanel* GraphPanel = GetOwnerPanel().Get())
-	{
-		if (GraphPanel->SelectionManager.SelectedNodes.Contains(GraphNode))
-		{
-			const FSlateBrush* ShadowBrush = FEditorStyle::GetBrush(TEXT("Graph.Node.ShadowSelected"));
-			const FVector2D NodeShadowSize = GetDefault<UGraphEditorSettings>()->GetShadowDeltaSize();
-
-			FSlateDrawElement::MakeBox(
-				OutDrawElements,
-				LayerId,
-				GetPaintSpaceGeometry().ToInflatedPaintGeometry(NodeShadowSize),
-				ShadowBrush,
-				ESlateDrawEffect::None,
-				FLinearColor(1.0f, 1.0f, 1.0f, 1.f)
-			);
-		}
-	}
+	//if(SGraphPanel* GraphPanel = GetOwnerPanel().Get())
+	//{
+	//	if (GraphPanel->SelectionManager.SelectedNodes.Contains(GraphNode))
+	//	{
+	//		const FSlateBrush* ShadowBrush = FEditorStyle::GetBrush(TEXT("Graph.Node.ShadowSelected"));
+	//		const FVector2D NodeShadowSize = GetDefault<UGraphEditorSettings>()->GetShadowDeltaSize();
+
+	//		FSlateDrawElement::MakeBox(
+	//			OutDrawElements,
+	//			LayerId,
+	//			GetPaintSpaceGeometry().ToInflatedPaintGeometry(NodeShadowSize),
+	//			ShadowBrush,
+	//			ESlateDrawEffect::None,
+	//			FLinearColor(1.0f, 1.0f, 1.0f, 1.f)
+	//		);
+	//	}
+	//}
 
 	return SGraphNode::OnPaint(Args, AllottedGeometry, MyCullingRect, OutDrawElements, LayerId, InWidgetStyle, bParentEnabled);
 }
@@ -412,9 +643,10 @@ void SDataprepGraphActionNode::OnDragEnter(const FGeometry& MyGeometry, const FD
 FSlateColor SDataprepGraphActionNode::GetInsertColor(int32 Index)
 {
 	static const FSlateColor BackgroundColor = FDataprepEditorStyle::GetColor("DataprepActionStep.BackgroundColor");
-	static const FSlateColor DragAndDrop = FDataprepEditorStyle::GetColor("DataprepActionStep.DragAndDrop");
+	static const FSlateColor DragAndDrop = FDataprepEditorStyle::GetColor("DataprepActionStep.Separator.Color")/*FDataprepEditorStyle::GetColor("DataprepActionStep.DragAndDrop")*/;
+
 
-	return Index == InsertIndex ? DragAndDrop : BackgroundColor;
+	return Index == InsertIndex ? DragAndDrop : FLinearColor::Transparent;
 }
 
 void SDataprepGraphActionNode::SetDraggedIndex(int32 Index)
@@ -435,6 +667,16 @@ void SDataprepGraphActionNode::SetHoveredIndex(int32 Index)
 	}
 }
 
+FMargin SDataprepGraphActionNode::GetOuterPadding() const
+{
+	static const FMargin Selected = FDataprepEditorStyle::GetMargin( "DataprepAction.Outter.Selected.Padding" );
+	static const FMargin Regular = FDataprepEditorStyle::GetMargin( "DataprepAction.Outter.Regular.Padding" );
+
+	const bool bIsSelected = GetOwnerPanel()->SelectionManager.SelectedNodes.Contains(GraphNode);
+
+	return bIsSelected ? Selected : Regular;
+}
+
 void SDataprepGraphActionNode::PopulateActionStepListWidget()
 {
 	if(!ActionStepListWidgetPtr.IsValid())
@@ -501,6 +743,7 @@ void SDataprepGraphActionNode::PopulateActionStepListWidget()
 	TSharedPtr<SDataprepEmptyActionStepNode> BottomSlot;
 	ActionStepListWidgetPtr->AddSlot()
 	.AutoHeight()
+	.Padding(0.f, 0.f, 0.f, 10.f)
 	[
 		SAssignNew(BottomSlot, SDataprepEmptyActionStepNode, SharedThis(this))
 	];
diff --git a/Engine/Plugins/Enterprise/DataprepEditor/Source/DataprepEditor/Private/Widgets/DataprepGraph/SDataprepGraphActionNode.h b/Engine/Plugins/Enterprise/DataprepEditor/Source/DataprepEditor/Private/Widgets/DataprepGraph/SDataprepGraphActionNode.h
--- a/Engine/Plugins/Enterprise/DataprepEditor/Source/DataprepEditor/Private/Widgets/DataprepGraph/SDataprepGraphActionNode.h
+++ b/Engine/Plugins/Enterprise/DataprepEditor/Source/DataprepEditor/Private/Widgets/DataprepGraph/SDataprepGraphActionNode.h
@@ -38,6 +38,7 @@ public:
 	// End of SWidget interface
 
 	// SGraphNode interface
+	virtual void UpdateGraphNode() override;
 	virtual TSharedRef<SWidget> CreateNodeContentArea() override;
 	virtual FReply OnDrop( const FGeometry& MyGeometry, const FDragDropEvent& DragDropEvent ) override;
 	virtual const FSlateBrush* GetShadowBrush(bool bSelected) const override;
@@ -59,9 +60,12 @@ public:
 	/** Set index of step node being dragged */
 	void SetDraggedIndex(int32 Index);
 
-	/** Set index of step node being hovered */
+	/** Get/Set index of step node being hovered */
+	int32 GetHoveredIndex() { return InsertIndex; }
 	void SetHoveredIndex(int32 Index);
 
+	static TSharedRef<SWidget> CreateBackground(const TAttribute<FSlateColor>& ColorAndOpacity);
+
 private:
 	/** Callback to handle changes on array of steps in action */
 	void OnStepsChanged();
@@ -69,6 +73,8 @@ private:
 	/** Reconstructs the list of widgets associated with the action's steps */
 	void PopulateActionStepListWidget();
 
+	FMargin GetOuterPadding() const;
+
 private:
 	/** Weak pointer to the associated action asset */
 	TWeakObjectPtr<class UDataprepActionAsset> DataprepActionPtr;
@@ -96,4 +102,6 @@ private:
 
 	/** Array of strong pointers to the UEdGraphNodes created for the action's steps */
 	TArray<TStrongObjectPtr<UDataprepGraphActionStepNode>> EdGraphStepNodes;
+
+	friend SDataprepGraphActionProxyNode;
 };
diff --git a/Engine/Plugins/Enterprise/DataprepEditor/Source/DataprepEditor/Private/Widgets/DataprepGraph/SDataprepGraphActionStepNode.cpp b/Engine/Plugins/Enterprise/DataprepEditor/Source/DataprepEditor/Private/Widgets/DataprepGraph/SDataprepGraphActionStepNode.cpp
--- a/Engine/Plugins/Enterprise/DataprepEditor/Source/DataprepEditor/Private/Widgets/DataprepGraph/SDataprepGraphActionStepNode.cpp
+++ b/Engine/Plugins/Enterprise/DataprepEditor/Source/DataprepEditor/Private/Widgets/DataprepGraph/SDataprepGraphActionStepNode.cpp
@@ -30,9 +30,9 @@
 #include "Widgets/Images/SImage.h"
 #include "Widgets/Layout/SBox.h"
 #include "Widgets/Layout/SScaleBox.h"
+#include "Widgets/Layout/SSeparator.h"
 #include "Widgets/SBoxPanel.h"
-#include "Widgets/Text/SInlineEditableTextBlock.h"
-#include "Widgets/Text/STextBlock.h"
+#include "Widgets/SOverlay.h"
 
 #define LOCTEXT_NAMESPACE "DataprepGraphEditor"
 
@@ -56,18 +56,40 @@ TSharedPtr<SWidget> SDataprepGraphActionStepNode::GetStepTitleWidget() const
 {
 	if(SDataprepActionBlock* ActionStepBlock = ActionStepBlockPtr.Get())
 	{
-		return SNew(SBorder)
-		.BorderImage( FEditorStyle::GetBrush( "Graph.StateNode.Body" ) )
-		.Padding(0.0f)
-		.BorderBackgroundColor( FDataprepEditorStyle::GetColor("DataprepActionStep.DragAndDrop") )
-		[
-			SNew(SBox)
-			.Padding( FMargin( 2.0f ) )
-			.Content()
+		TAttribute<FSlateColor> BlockColorAndOpacity = TAttribute<FSlateColor>::Create(TAttribute<FSlateColor>::FGetter::CreateSP(this, &SDataprepGraphActionStepNode::GetBlockOverlayColor));
+
+		return SNew(SOverlay)
+			+ SOverlay::Slot()
+			.Padding(FDataprepEditorStyle::GetMargin( "DataprepActionStep.DnD.Outter.Padding" ))
+			.VAlign(VAlign_Fill)
+			.HAlign(HAlign_Fill)
+			[
+				SDataprepGraphActionNode::CreateBackground(BlockColorAndOpacity)
+			]
+
+			+ SOverlay::Slot()
+			.Padding(FDataprepEditorStyle::GetMargin( "DataprepActionStep.DnD.Inner.Padding" ))
+			.VAlign(VAlign_Fill)
+			.HAlign(HAlign_Fill)
 			[
-				ActionStepBlock->GetBlockTitleWidget()
+				SDataprepGraphActionNode::CreateBackground(FDataprepEditorStyle::GetColor( "DataprepActionStep.BackgroundColor" ))
 			]
-		];
+
+			+ SOverlay::Slot()
+			.Padding(FDataprepEditorStyle::GetMargin( "DataprepActionStep.DnD.Inner.Padding" ))
+			.VAlign(VAlign_Fill)
+			.HAlign(HAlign_Fill)
+			[
+				SNew(SVerticalBox)
+
+				+SVerticalBox::Slot()
+				.AutoHeight()
+				.Padding(5.f, 10.f)
+				.VAlign(VAlign_Center)
+				[
+					ActionStepBlock->GetTitleWidget()
+				]
+			];
 	}
 
 	return TSharedPtr<SWidget>();
@@ -81,8 +103,7 @@ void SDataprepGraphActionStepNode::UpdateGraphNode()
 	RightNodeBox.Reset();
 	LeftNodeBox.Reset();
 
-	TSharedRef<SWidget> ActionBlock = SNullWidget::NullWidget;
-
+	TSharedRef<SWidget> ActionBlockPtr = SNullWidget::NullWidget;
 	if(UDataprepGraphActionStepNode* ActionStepNode = Cast<UDataprepGraphActionStepNode>(GraphNode))
 	{
 		TSharedRef< FDataprepSchemaActionContext > StepData = MakeShared< FDataprepSchemaActionContext >();
@@ -97,17 +118,17 @@ void SDataprepGraphActionStepNode::UpdateGraphNode()
 			if (StepType == UDataprepOperation::StaticClass())
 			{
 				UDataprepOperation* Operation = static_cast<UDataprepOperation*>( StepObject );
-				ActionStepBlockPtr = StaticCastSharedRef<SDataprepActionBlock>( SNew(SDataprepOperation, Operation, StepData) );
+				ActionStepBlockPtr = StaticCastSharedRef<SDataprepActionBlock>( SNew(SDataprepOperation, Operation, StepData).IsSimplified(true) );
 			}
 			else if (StepType == UDataprepFilter::StaticClass())
 			{
 				UDataprepFilter* Filter = static_cast<UDataprepFilter*>( StepObject );
-				ActionStepBlockPtr = StaticCastSharedRef<SDataprepActionBlock>( SNew(SDataprepFilter, *Filter, StepData) );
+				ActionStepBlockPtr = StaticCastSharedRef<SDataprepActionBlock>( SNew(SDataprepFilter, *Filter, StepData).IsSimplified(true) );
 			}
 			else if (StepType == UDataprepSelectionTransform::StaticClass())
 			{
 				UDataprepSelectionTransform* SelectionTransform = static_cast<UDataprepSelectionTransform*>( StepObject );
-				ActionStepBlockPtr = StaticCastSharedRef<SDataprepActionBlock>( SNew(SDataprepSelectionTransform, SelectionTransform, StepData) );
+				ActionStepBlockPtr = StaticCastSharedRef<SDataprepActionBlock>( SNew(SDataprepSelectionTransform, SelectionTransform, StepData).IsSimplified(true) );
 			}
 
 			if(ActionStepBlockPtr.IsValid())
@@ -115,11 +136,14 @@ void SDataprepGraphActionStepNode::UpdateGraphNode()
 #ifndef NO_BLUEPRINT
 				ActionStepBlockPtr->bIsSimplifiedGraph = true;
 #endif
-				ActionBlock = ActionStepBlockPtr->AsShared();
+				ActionBlockPtr = ActionStepBlockPtr->AsShared();
 			}
 		}
 	}
 
+	TAttribute<FMargin> OverlayPadding = TAttribute<FMargin>::Create(TAttribute<FMargin>::FGetter::CreateSP(this, &SDataprepGraphActionStepNode::GetBlockPadding));
+	TAttribute<FSlateColor> BlockColorAndOpacity = TAttribute<FSlateColor>::Create(TAttribute<FSlateColor>::FGetter::CreateSP(this, &SDataprepGraphActionStepNode::GetBlockOverlayColor));
+
 	this->ContentScale.Bind( this, &SGraphNode::GetContentScale );
 	this->GetOrAddSlot( ENodeZone::Center )
 	.HAlign(HAlign_Fill)
@@ -129,36 +153,64 @@ void SDataprepGraphActionStepNode::UpdateGraphNode()
 
 		+SVerticalBox::Slot()
 		.AutoHeight()
-		.Padding(10.f, 0.f, 10.f, 0.f)
+		.Padding(20.f, 0.f)
 		[
-			SNew(SBorder)
-			.BorderBackgroundColor( this, &SDataprepGraphActionStepNode::GetDragAndDropColor )
-			.BorderImage(FEditorStyle::GetBrush("Graph.StateNode.Body"))
-			[
-				SNew(SBox)
-				.HeightOverride(2.0)
-			]
+			SNew( SSeparator )
+			.SeparatorImage(FEditorStyle::GetBrush( "ThinLine.Horizontal" ))
+			.Thickness(2.f)
+			.Orientation(EOrientation::Orient_Horizontal)
+			.ColorAndOpacity(this, &SDataprepGraphActionStepNode::GetDragAndDropColor)
 		]
 
 		+SVerticalBox::Slot()
 		.AutoHeight()
 		[
-			SNew(SBorder)
-			.BorderImage( FEditorStyle::GetBrush( "Graph.StateNode.Body" ) )
-			.Padding(2.0f)
-			.BorderBackgroundColor( this, &SDataprepGraphActionStepNode::GetBorderBackgroundColor )
+			SNew(SOverlay)
+
+			+ SOverlay::Slot()
+			.Padding(OverlayPadding)
+			.VAlign(VAlign_Fill)
+			.HAlign(HAlign_Fill)
 			[
-				SNew(SBox)
-				.Padding( FMargin( 1.0f ) )
-				.Content()
-				[
-					ActionBlock
-				]
+				SDataprepGraphActionNode::CreateBackground(BlockColorAndOpacity)
+			]
+
+			+ SOverlay::Slot()
+			.Padding(FDataprepEditorStyle::GetMargin( "DataprepActionStep.Padding" ))
+			.VAlign(VAlign_Fill)
+			.HAlign(HAlign_Fill)
+			[
+				SDataprepGraphActionNode::CreateBackground(FDataprepEditorStyle::GetColor( "DataprepActionStep.BackgroundColor" ))
+			]
+
+			+ SOverlay::Slot()
+			.Padding(FDataprepEditorStyle::GetMargin( "DataprepActionStep.Padding" ))
+			.VAlign(VAlign_Fill)
+			.HAlign(HAlign_Fill)
+			[
+				ActionBlockPtr
 			]
 		]
 	];
 }
 
+FSlateColor SDataprepGraphActionStepNode::GetBlockOverlayColor() const
+{
+	static const FSlateColor BlockColor = FDataprepEditorStyle::GetColor( "DataprepActionStep.Filter.OutlineColor" );
+
+	return ActionStepBlockPtr.IsValid() ? ActionStepBlockPtr->GetOutlineColor() : BlockColor;
+}
+
+FMargin SDataprepGraphActionStepNode::GetBlockPadding()
+{
+	static const FMargin Selected = FDataprepEditorStyle::GetMargin( "DataprepActionStep.Outter.Selected.Padding" );
+	static const FMargin Regular = FDataprepEditorStyle::GetMargin( "DataprepActionStep.Outter.Regular.Padding" );
+
+	const bool bIsSelected = GetOwnerPanel()->SelectionManager.SelectedNodes.Contains(GraphNode);
+
+	return bIsSelected ? Selected : Regular;
+}
+
 FSlateColor SDataprepGraphActionStepNode::GetDragAndDropColor() const
 {
 	return ParentNodePtr.Pin()->GetInsertColor(StepIndex);
diff --git a/Engine/Plugins/Enterprise/DataprepEditor/Source/DataprepEditor/Private/Widgets/DataprepGraph/SDataprepGraphActionStepNode.h b/Engine/Plugins/Enterprise/DataprepEditor/Source/DataprepEditor/Private/Widgets/DataprepGraph/SDataprepGraphActionStepNode.h
--- a/Engine/Plugins/Enterprise/DataprepEditor/Source/DataprepEditor/Private/Widgets/DataprepGraph/SDataprepGraphActionStepNode.h
+++ b/Engine/Plugins/Enterprise/DataprepEditor/Source/DataprepEditor/Private/Widgets/DataprepGraph/SDataprepGraphActionStepNode.h
@@ -39,6 +39,8 @@ public:
 	virtual FReply OnDragOver(const FGeometry& MyGeometry, const FDragDropEvent& DragDropEvent) override;
 	virtual void OnDragLeave(const FDragDropEvent& DragDropEvent) override;
 	virtual FReply OnDrop( const FGeometry& MyGeometry, const FDragDropEvent& DragDropEvent ) override;
+	virtual bool IsNameReadOnly () const override { return true; }
+	virtual FSlateRect GetTitleRect() const override { return FSlateRect(); }
 	// End of SGraphNode interface
 
 	int32 GetStepIndex() const { return StepIndex; }
@@ -59,6 +61,10 @@ private:
 	 */
 	FSlateColor GetDragAndDropColor() const;
 
+	FSlateColor GetBlockOverlayColor() const;
+
+	FMargin GetBlockPadding();
+
 private:
 	/** Pointer to the widget displaying the actual filter or operation */
 	TSharedPtr<SDataprepActionBlock> ActionStepBlockPtr;
diff --git a/Engine/Plugins/Enterprise/DataprepEditor/Source/DataprepEditor/Private/Widgets/DataprepGraph/SDataprepOperation.cpp b/Engine/Plugins/Enterprise/DataprepEditor/Source/DataprepEditor/Private/Widgets/DataprepGraph/SDataprepOperation.cpp
--- a/Engine/Plugins/Enterprise/DataprepEditor/Source/DataprepEditor/Private/Widgets/DataprepGraph/SDataprepOperation.cpp
+++ b/Engine/Plugins/Enterprise/DataprepEditor/Source/DataprepEditor/Private/Widgets/DataprepGraph/SDataprepOperation.cpp
@@ -18,10 +18,10 @@ void SDataprepOperation::Construct(const FArguments& InArgs, UDataprepOperation*
 	TAttribute<FText> TooltipTextAttribute = MakeAttributeSP( this, &SDataprepOperation::GetTooltipText );
 	SetToolTipText( TooltipTextAttribute );
 
-	SDataprepActionBlock::Construct( SDataprepActionBlock::FArguments(), InDataprepActionContext );
+	SDataprepActionBlock::Construct( SDataprepActionBlock::FArguments().IsSimplified(InArgs._IsSimplified), InDataprepActionContext );
 }
 
-FLinearColor SDataprepOperation::GetOutlineColor() const
+FSlateColor SDataprepOperation::GetOutlineColor() const
 {
 	return FDataprepEditorStyle::GetColor( "DataprepActionStep.Operation.OutlineColor" );
 }
diff --git a/Engine/Plugins/Enterprise/DataprepEditor/Source/DataprepEditor/Private/Widgets/DataprepGraph/SDataprepOperation.h b/Engine/Plugins/Enterprise/DataprepEditor/Source/DataprepEditor/Private/Widgets/DataprepGraph/SDataprepOperation.h
--- a/Engine/Plugins/Enterprise/DataprepEditor/Source/DataprepEditor/Private/Widgets/DataprepGraph/SDataprepOperation.h
+++ b/Engine/Plugins/Enterprise/DataprepEditor/Source/DataprepEditor/Private/Widgets/DataprepGraph/SDataprepOperation.h
@@ -14,15 +14,22 @@ struct FDataprepSchemaActionContext;
 
 class SDataprepOperation : public SDataprepActionBlock, public FGCObject
 {
+#ifndef NO_BLUEPRINT
+	SLATE_BEGIN_ARGS(SDataprepOperation)
+	: _IsSimplified(false)
+	{}
+		SLATE_ARGUMENT( bool, IsSimplified )
+#else
 	SLATE_BEGIN_ARGS(SDataprepOperation) {}
-	SLATE_END_ARGS();
+#endif
+	SLATE_END_ARGS()
 
 	void Construct(const FArguments& InArgs, UDataprepOperation* InOperation, const TSharedRef<FDataprepSchemaActionContext>& InDataprepActionContext);
 
 protected:
 
 	// SDataprepActionBlock interface
-	virtual FLinearColor GetOutlineColor() const override;
+	virtual FSlateColor GetOutlineColor() const override;
 	virtual FText GetBlockTitle() const override;
 	virtual TSharedRef<SWidget> GetContentWidget() override;
 	//~ end of SDataprepActionBlock interface
diff --git a/Engine/Plugins/Enterprise/DataprepEditor/Source/DataprepEditor/Private/Widgets/DataprepGraph/SDataprepSelectionTransform.cpp b/Engine/Plugins/Enterprise/DataprepEditor/Source/DataprepEditor/Private/Widgets/DataprepGraph/SDataprepSelectionTransform.cpp
--- a/Engine/Plugins/Enterprise/DataprepEditor/Source/DataprepEditor/Private/Widgets/DataprepGraph/SDataprepSelectionTransform.cpp
+++ b/Engine/Plugins/Enterprise/DataprepEditor/Source/DataprepEditor/Private/Widgets/DataprepGraph/SDataprepSelectionTransform.cpp
@@ -18,10 +18,10 @@ void SDataprepSelectionTransform::Construct(const FArguments& InArgs, UDataprepS
 	TAttribute<FText> TooltipTextAttribute = MakeAttributeSP(this, &SDataprepSelectionTransform::GetTooltipText);
 	SetToolTipText(TooltipTextAttribute);
 
-	SDataprepActionBlock::Construct(SDataprepActionBlock::FArguments(), InDataprepActionContext);
+	SDataprepActionBlock::Construct(SDataprepActionBlock::FArguments().IsSimplified(InArgs._IsSimplified), InDataprepActionContext);
 }
 
-FLinearColor SDataprepSelectionTransform::GetOutlineColor() const
+FSlateColor SDataprepSelectionTransform::GetOutlineColor() const
 {
 	return FDataprepEditorStyle::GetColor("DataprepActionStep.Operation.OutlineColor");
 }
diff --git a/Engine/Plugins/Enterprise/DataprepEditor/Source/DataprepEditor/Private/Widgets/DataprepGraph/SDataprepSelectionTransform.h b/Engine/Plugins/Enterprise/DataprepEditor/Source/DataprepEditor/Private/Widgets/DataprepGraph/SDataprepSelectionTransform.h
--- a/Engine/Plugins/Enterprise/DataprepEditor/Source/DataprepEditor/Private/Widgets/DataprepGraph/SDataprepSelectionTransform.h
+++ b/Engine/Plugins/Enterprise/DataprepEditor/Source/DataprepEditor/Private/Widgets/DataprepGraph/SDataprepSelectionTransform.h
@@ -14,7 +14,10 @@ struct FDataprepSchemaActionContext;
 
 class SDataprepSelectionTransform : public SDataprepActionBlock, public FGCObject
 {
-	SLATE_BEGIN_ARGS(SDataprepSelectionTransform) {}
+	SLATE_BEGIN_ARGS(SDataprepSelectionTransform)
+	: _IsSimplified(false)
+	{}
+		SLATE_ARGUMENT( bool, IsSimplified )
 	SLATE_END_ARGS();
 
 	void Construct(const FArguments& InArgs, UDataprepSelectionTransform* InTransform, const TSharedRef<FDataprepSchemaActionContext>& InDataprepActionContext);
@@ -22,7 +25,7 @@ class SDataprepSelectionTransform : public SDataprepActionBlock, public FGCObjec
 protected:
 
 	// SDataprepActionBlock interface
-	virtual FLinearColor GetOutlineColor() const override;
+	virtual FSlateColor GetOutlineColor() const override;
 	virtual FText GetBlockTitle() const override;
 	virtual TSharedRef<SWidget> GetContentWidget() override;
 	//~ end of SDataprepActionBlock interface
