Fixed crash when F2 is pressed to change name of action node

#rnx

diff --git a/Engine/Plugins/Enterprise/DataprepEditor/Source/DataprepEditor/Private/DataprepGraph/DataprepGraphActionNode.h b/Engine/Plugins/Enterprise/DataprepEditor/Source/DataprepEditor/Private/DataprepGraph/DataprepGraphActionNode.h
--- a/Engine/Plugins/Enterprise/DataprepEditor/Source/DataprepEditor/Private/DataprepGraph/DataprepGraphActionNode.h
+++ b/Engine/Plugins/Enterprise/DataprepEditor/Source/DataprepEditor/Private/DataprepGraph/DataprepGraphActionNode.h
@@ -82,7 +82,7 @@ public:
 	const UDataprepActionAsset* GetDataprepActionAsset() const { return DataprepActionAsset; }
 	UDataprepActionAsset* GetDataprepActionAsset() { return DataprepActionAsset; }
 
-	int32 GetExecutionOrder() { return ExecutionOrder; }
+	int32 GetExecutionOrder() const { return ExecutionOrder; }
 
 protected:
 	UPROPERTY()
diff --git a/Engine/Plugins/Enterprise/DataprepEditor/Source/DataprepEditor/Private/Widgets/DataprepGraph/SDataprepGraphActionNode.cpp b/Engine/Plugins/Enterprise/DataprepEditor/Source/DataprepEditor/Private/Widgets/DataprepGraph/SDataprepGraphActionNode.cpp
--- a/Engine/Plugins/Enterprise/DataprepEditor/Source/DataprepEditor/Private/Widgets/DataprepGraph/SDataprepGraphActionNode.cpp
+++ b/Engine/Plugins/Enterprise/DataprepEditor/Source/DataprepEditor/Private/Widgets/DataprepGraph/SDataprepGraphActionNode.cpp
@@ -134,16 +134,23 @@ class SDataprepEmptyActionStepNode : public SVerticalBox
 	};
 
 public:
-	SLATE_BEGIN_ARGS(SDataprepEmptyActionStepNode) {}
+	SLATE_BEGIN_ARGS(SDataprepEmptyActionStepNode)
+		: _Text(LOCTEXT("DataprepEmptyActionStepLabel", "+ Add Step"))
+		{}
+
+		/** The text displayed in the text block */
+		SLATE_ATTRIBUTE( FText, Text )
+		SLATE_ATTRIBUTE( int32, StepIndex )
 	SLATE_END_ARGS();
 
 	void Construct(const FArguments& InArgs, const TSharedPtr<SDataprepGraphActionNode>& InParent)
 	{
 		ParentPtr = InParent;
+		StepIndex = InArgs._StepIndex.Get(InParent->GetDataprepAction()->GetStepsCount());
+		bIsHovered = false;
 
 		const float InterStepSpacing = 5.f;
-
-		bIsHovered = false;
+		const bool bBottomSlot = StepIndex == InParent->GetDataprepAction()->GetStepsCount();
 
 		SVerticalBox::Construct(SVerticalBox::FArguments());
 
@@ -151,51 +158,85 @@ public:
 		TAttribute<FSlateColor> OuterColorAndOpacity = TAttribute<FSlateColor>::Create(TAttribute<FSlateColor>::FGetter::CreateSP(this, &SDataprepEmptyActionStepNode::GetColor, EColorType::OuterColor));
 		TAttribute<FSlateColor> InnerColorAndOpacity = TAttribute<FSlateColor>::Create(TAttribute<FSlateColor>::FGetter::CreateSP(this, &SDataprepEmptyActionStepNode::GetColor, EColorType::InnerColor));
 
+		TSharedPtr<SWidget> Separator = SNullWidget::NullWidget;
+
+		if(StepIndex != InParent->GetDataprepAction()->GetStepsCount())
+		{
+			Separator = SNew( SSeparator )
+				.SeparatorImage(FEditorStyle::GetBrush( "ThinLine.Horizontal" ))
+				.Thickness(2.f)
+				.Orientation(EOrientation::Orient_Horizontal)
+				.ColorAndOpacity(this, &SDataprepEmptyActionStepNode::GetDragAndDropColor);
+		}
+
 		AddSlot()
 		.AutoHeight()
 		.Padding(FDataprepEditorStyle::GetMargin( "DataprepActionStep.Padding" ))
 		[
-			SNew(SOverlay)
+			SNew(SVerticalBox)
 
-			+ SOverlay::Slot()
-			.Padding(0.f)
-			.VAlign(VAlign_Fill)
-			.HAlign(HAlign_Fill)
+			+SVerticalBox::Slot()
+			.Padding(20.f, 0.f)
+			.AutoHeight()
 			[
-				SNew(SImage)
-				.ColorAndOpacity(MoveTemp(OuterColorAndOpacity))
-				.Image(FEditorStyle::GetBrush( "Graph.StateNode.Body" ))
+				Separator.ToSharedRef()
 			]
 
-			+ SOverlay::Slot()
-			.Padding(1.f)
-			.VAlign(VAlign_Fill)
-			.HAlign(HAlign_Fill)
+			+SVerticalBox::Slot()
+			.AutoHeight()
 			[
-				SNew(SImage)
-				.ColorAndOpacity(MoveTemp(InnerColorAndOpacity))
-				.Image(FEditorStyle::GetBrush( "Graph.StateNode.Body" ))
-			]
+				SNew(SOverlay)
 
-			+ SOverlay::Slot()
-			.Padding(10.f)
-			.VAlign(VAlign_Fill)
-			.HAlign(HAlign_Fill)
-			[
-				SNew(SHorizontalBox)
+				+ SOverlay::Slot()
+				.Padding(10.f)
+				.VAlign(VAlign_Fill)
+				.HAlign(HAlign_Fill)
+				[
+					SNew(SBox)
+					.MinDesiredWidth(250.f)
+				]
+
+				+ SOverlay::Slot()
+				.Padding(0.f)
+				.VAlign(VAlign_Fill)
+				.HAlign(HAlign_Fill)
+				[
+					SNew(SImage)
+					.ColorAndOpacity(MoveTemp(OuterColorAndOpacity))
+					.Image(FEditorStyle::GetBrush( "Graph.StateNode.Body" ))
+				]
 
-				+ SHorizontalBox::Slot()
+				+ SOverlay::Slot()
+				.Padding(1.f)
 				.VAlign(VAlign_Fill)
 				.HAlign(HAlign_Fill)
 				[
-					SNew(STextBlock)
-					.Text(LOCTEXT("DataprepEmptyActionStepLabel", "+ Add Step"))
-					.TextStyle( &FDataprepEditorStyle::GetWidgetStyle<FTextBlockStyle>( "DataprepActionBlock.TitleTextBlockStyle" ) )
-					.ColorAndOpacity(MoveTemp(TextColorAndOpacity))
-					.Justification(ETextJustify::Center)
+					SNew(SImage)
+					.ColorAndOpacity(MoveTemp(InnerColorAndOpacity))
+					.Image(FEditorStyle::GetBrush( "Graph.StateNode.Body" ))
+				]
+
+				+ SOverlay::Slot()
+				.Padding(10.f, 10.f, 10.f, bBottomSlot ? 10.f : 5.f)
+				.VAlign(VAlign_Fill)
+				.HAlign(HAlign_Fill)
+				[
+					SNew(SHorizontalBox)
+
+					+ SHorizontalBox::Slot()
+					.VAlign(VAlign_Fill)
+					.HAlign(HAlign_Fill)
+					[
+						SAssignNew(TextBlock, STextBlock)
+						.TextStyle( &FDataprepEditorStyle::GetWidgetStyle<FTextBlockStyle>( "DataprepActionBlock.TitleTextBlockStyle" ) )
+						.ColorAndOpacity(MoveTemp(TextColorAndOpacity))
+						.Justification(ETextJustify::Center)
+					]
 				]
 			]
 		];
+
+		TextBlock->SetText(InArgs._Text);
 	}
 
 	// SWidget Interface
@@ -207,7 +248,7 @@ public:
 			ParentTrackNodePtr.Pin()->OnDragLeave(DragDropEvent);
 
 			DragActionStepNodeOp->SetHoveredNode(ParentPtr.Pin()->GetNodeObj());
-			ParentPtr.Pin()->SetHoveredIndex( ParentPtr.Pin()->GetDataprepAction()->GetStepsCount() );
+			ParentPtr.Pin()->SetHoveredIndex( StepIndex );
 		}
 
 		SVerticalBox::OnDragEnter(MyGeometry, DragDropEvent);
@@ -219,7 +260,7 @@ public:
 		if(DragActionStepNodeOp.IsValid() && ParentPtr.IsValid())
 		{
 			DragActionStepNodeOp->SetHoveredNode(ParentPtr.Pin()->GetNodeObj());
-			ParentPtr.Pin()->SetHoveredIndex( ParentPtr.Pin()->GetDataprepAction()->GetStepsCount() );
+			ParentPtr.Pin()->SetHoveredIndex( StepIndex );
 
 			return FReply::Handled();
 		}
@@ -249,8 +290,15 @@ public:
 		TSharedPtr<FDataprepDragDropOp> DragActionStepNodeOp = DragDropEvent.GetOperationAs<FDataprepDragDropOp>();
 		if (DragActionStepNodeOp.IsValid() && ParentPtr.IsValid())
 		{
-			const FVector2D NodeAddPosition = ParentPtr.Pin()->NodeCoordToGraphCoord( MyGeometry.AbsoluteToLocal( DragDropEvent.GetScreenSpacePosition() ) );
-			return DragActionStepNodeOp->DroppedOnNode(DragDropEvent.GetScreenSpacePosition(), NodeAddPosition);
+			if(StepIndex == ParentPtr.Pin()->GetDataprepAction()->GetStepsCount())
+			{
+				const FVector2D NodeAddPosition = ParentPtr.Pin()->NodeCoordToGraphCoord( MyGeometry.AbsoluteToLocal( DragDropEvent.GetScreenSpacePosition() ) );
+				return DragActionStepNodeOp->DroppedOnNode(DragDropEvent.GetScreenSpacePosition(), NodeAddPosition);
+			}
+			else
+			{
+				return FReply::Handled().EndDragDrop();
+			}
 		}
 
 		return SVerticalBox::OnDrop(MyGeometry, DragDropEvent);
@@ -262,6 +310,14 @@ public:
 		ParentTrackNodePtr = InParentTrackNode;
 	}
 
+	void SetText(FText Text)
+	{
+		if(TextBlock.IsValid())
+		{
+			TextBlock->SetText(Text);
+		}
+	}
+
 private:
 	FSlateColor GetColor(EColorType Type) const
 	{
@@ -276,14 +332,19 @@ private:
 			FDataprepEditorStyle::GetColor("DataprepAction.EmptyStep.Background.Normal"),
 		};
 
-		const int32 EmptyIndex = ParentPtr.Pin()->GetDataprepAction()->GetStepsCount();
+		return StepIndex == ParentPtr.Pin()->GetHoveredIndex() ? HoveredColors[Type] : NormalColors[Type];
+	}
 
-		return EmptyIndex == ParentPtr.Pin()->GetHoveredIndex() ? HoveredColors[Type] : NormalColors[Type];
+	FSlateColor GetDragAndDropColor() const
+	{
+		return ParentPtr.Pin()->GetInsertColor(StepIndex);
 	}
 
 private:
 	TWeakPtr<SDataprepGraphActionNode> ParentPtr;
 	TWeakPtr<SDataprepGraphTrackNode> ParentTrackNodePtr;
+	TSharedPtr<STextBlock> TextBlock;
+	int32 StepIndex;
 };
 
 void SDataprepGraphActionNode::Construct(const FArguments& InArgs, UDataprepGraphActionNode* InActionNode)
@@ -312,14 +373,23 @@ void SDataprepGraphActionNode::SetParentTrackNode(TSharedPtr<SDataprepGraphTrack
 
 	if(ActionStepListWidgetPtr.IsValid())
 	{
+		FChildren* StepListChildren = ActionStepListWidgetPtr->GetChildren();
+
 		// Update parent track on step widgets
 		for(TSharedPtr<SDataprepGraphActionStepNode>& ActionStepGraphNode : ActionStepGraphNodes)
 		{
-			ActionStepGraphNode->SetParentTrackNode(InParentTrackNode);
+			if(ActionStepGraphNode->GetStepTitleWidget().IsValid())
+			{
+				ActionStepGraphNode->SetParentTrackNode(InParentTrackNode);
+			}
+			else
+			{
+				TSharedRef<SDataprepEmptyActionStepNode> EmptyWidgetPtr = StaticCastSharedRef<SDataprepEmptyActionStepNode>(StepListChildren->GetChildAt(ActionStepGraphNode->GetStepIndex()));
+				EmptyWidgetPtr->SetParentTrackNode(InParentTrackNode);
+			}
 		}
 
 		// Update parent on empty bottom widget
-		FChildren* StepListChildren = ActionStepListWidgetPtr->GetChildren();
 		TSharedRef<SDataprepEmptyActionStepNode> EmptyWidgetPtr = StaticCastSharedRef<SDataprepEmptyActionStepNode>(StepListChildren->GetChildAt(StepListChildren->Num() - 1));
 		EmptyWidgetPtr->SetParentTrackNode(InParentTrackNode);
 	}
@@ -723,11 +793,32 @@ void SDataprepGraphActionNode::PopulateActionStepListWidget()
 			ActionStepGraphNode->SetParentTrackNode(TrackNodePtr);
 		}
 
-		ActionStepListWidgetPtr->AddSlot()
-		.AutoHeight()
-		[
-			ActionStepGraphNode.ToSharedRef()
-		];
+		if(ActionStepGraphNode->GetStepTitleWidget().IsValid())
+		{
+			ActionStepListWidgetPtr->AddSlot()
+			.AutoHeight()
+			[
+				ActionStepGraphNode.ToSharedRef()
+			];
+		}
+		else
+		{
+			TSharedPtr<SDataprepEmptyActionStepNode> EmptySlot;
+
+			ActionStepListWidgetPtr->AddSlot()
+			.AutoHeight()
+			.Padding(0.f, 0.f, 0.f, 10.f)
+			[
+				SAssignNew(EmptySlot, SDataprepEmptyActionStepNode, SharedThis(this))
+				.Text(LOCTEXT("DataprepUnknownActionStepLabel", "Unknown Step Class"))
+				.StepIndex(Index)
+			];
+
+			if(TrackNodePtr.IsValid())
+			{
+				EmptySlot->SetParentTrackNode(TrackNodePtr);
+			}
+		}
 
 		ActionStepGraphNodes[Index] = ActionStepGraphNode;
 	}
diff --git a/Engine/Plugins/Enterprise/DataprepEditor/Source/DataprepEditor/Private/Widgets/DataprepGraph/SDataprepGraphEditor.cpp b/Engine/Plugins/Enterprise/DataprepEditor/Source/DataprepEditor/Private/Widgets/DataprepGraph/SDataprepGraphEditor.cpp
--- a/Engine/Plugins/Enterprise/DataprepEditor/Source/DataprepEditor/Private/Widgets/DataprepGraph/SDataprepGraphEditor.cpp
+++ b/Engine/Plugins/Enterprise/DataprepEditor/Source/DataprepEditor/Private/Widgets/DataprepGraph/SDataprepGraphEditor.cpp
@@ -420,9 +420,9 @@ void SDataprepGraphEditor::OnRenameNode()
 	for (FGraphPanelSelectionSet::TConstIterator NodeIt(SelectedNodes); NodeIt; ++NodeIt)
 	{
 		UEdGraphNode* SelectedNode = Cast<UEdGraphNode>(*NodeIt);
-		if (SelectedNode != NULL && SelectedNode->bCanRenameNode)
+		if (SelectedNode != NULL && SelectedNode->bCanRenameNode && TrackGraphNodePtr.IsValid())
 		{
-			IsNodeTitleVisible(SelectedNode, true);
+			TrackGraphNodePtr.Pin()->RequestRename(SelectedNode);
 			break;
 		}
 	}
diff --git a/Engine/Plugins/Enterprise/DataprepEditor/Source/DataprepEditor/Private/Widgets/DataprepGraph/SDataprepGraphTrackNode.cpp b/Engine/Plugins/Enterprise/DataprepEditor/Source/DataprepEditor/Private/Widgets/DataprepGraph/SDataprepGraphTrackNode.cpp
--- a/Engine/Plugins/Enterprise/DataprepEditor/Source/DataprepEditor/Private/Widgets/DataprepGraph/SDataprepGraphTrackNode.cpp
+++ b/Engine/Plugins/Enterprise/DataprepEditor/Source/DataprepEditor/Private/Widgets/DataprepGraph/SDataprepGraphTrackNode.cpp
@@ -529,6 +529,41 @@ void SDataprepGraphTrackNode::OnControlKeyChanged(bool bControlKeyDown)
 	}
 }
 
+void SDataprepGraphTrackNode::RequestRename(const UEdGraphNode* Node)
+{
+	if(SGraphPanel* GraphPanel = OwnerGraphPanelPtr.Pin().Get())
+	{
+		if(const UDataprepGraphActionNode* ActionEdNode = Cast<UDataprepGraphActionNode>(Node))
+		{
+			int32 ActionIndex = ActionEdNode->GetExecutionOrder();
+			if(ActionNodes.IsValidIndex(ActionIndex))
+			{
+				TSharedPtr<SDataprepGraphActionNode>& ActionNode = ActionNodes[ActionIndex];
+
+				if(ActionNode.IsValid() && !GraphPanel->HasMouseCapture())
+				{
+					FSlateRect TitleRect = ActionNode->GetTitleRect();
+					const FVector2D TopLeft = FVector2D( TitleRect.Left, TitleRect.Top );
+					const FVector2D BottomRight = FVector2D( TitleRect.Right, TitleRect.Bottom );
+
+					bool bTitleVisible = GraphPanel->IsRectVisible( TopLeft, BottomRight );
+					if( !bTitleVisible )
+					{
+						bTitleVisible = GraphPanel->JumpToRect( TopLeft, BottomRight );
+					}
+
+					if( bTitleVisible )
+					{
+						ActionNode->RequestRename();
+						//GraphPanel->JumpToNode(Node, false, true);
+						ActionNode->ApplyRename();
+					}
+				}
+			}
+		}
+	}
+}
+
 void SDataprepGraphTrackWidget::Construct(const FArguments& InArgs, TSharedPtr<SDataprepGraphTrackNode> InTrackNode)
 {
 	TrackNodePtr = InTrackNode;
diff --git a/Engine/Plugins/Enterprise/DataprepEditor/Source/DataprepEditor/Private/Widgets/DataprepGraph/SDataprepGraphTrackNode.h b/Engine/Plugins/Enterprise/DataprepEditor/Source/DataprepEditor/Private/Widgets/DataprepGraph/SDataprepGraphTrackNode.h
--- a/Engine/Plugins/Enterprise/DataprepEditor/Source/DataprepEditor/Private/Widgets/DataprepGraph/SDataprepGraphTrackNode.h
+++ b/Engine/Plugins/Enterprise/DataprepEditor/Source/DataprepEditor/Private/Widgets/DataprepGraph/SDataprepGraphTrackNode.h
@@ -110,6 +110,9 @@ public:
 	/** Recomputes the position of each action node */
 	bool RefreshLayout();
 
+	/** Start editing of action asset associated to input EdGraphNode */
+	void RequestRename(const UEdGraphNode* Node);
+
 	/** Miscellaneous values used in the display */
 	// #ueent_wip: Will be moved to the Dataprep editor's style
 	static FMargin NodePadding;
@@ -124,7 +127,7 @@ private:
 	TSharedPtr<SDataprepGraphTrackWidget> TrackWidgetPtr;
 
 	/** Array of action node's widgets */
-	mutable TArray<TSharedPtr<SDataprepGraphActionNode>> ActionNodes;
+	TArray<TSharedPtr<SDataprepGraphActionNode>> ActionNodes;
 
 	/** Weak pointer to the Dataprep asset holding the displayed actions */
 	TWeakObjectPtr<UDataprepAsset> DataprepAssetPtr;
