Integrating the dataprep preview system.

The way the dataprep objects notify the systems that they were edited was change for the need of the preview system. The parameterization code was updated.
The asset preview widget was refactor to support columns with a custom sorting. This also put it in a better to support custom filter at some point.
Functions of the dataprep filters were modified to process on an array view instead of a full array. Also added some functions to output the result and some additional info for the preview system to some buffer instead of returning an array.
Updated the dataprep editor so that our graph widgets can have a optional weak reference to the dataprep editor. (While at it I also fixed alpha issue of one of the images we are using for our new nodes)

This doesn't include the ui and the code to clear a preview.

diff --git a/Engine/Plugins/Enterprise/DataprepEditor/Source/DataprepCore/Private/DataprepActionAsset.cpp b/Engine/Plugins/Enterprise/DataprepEditor/Source/DataprepCore/Private/DataprepActionAsset.cpp
--- a/Engine/Plugins/Enterprise/DataprepEditor/Source/DataprepCore/Private/DataprepActionAsset.cpp
+++ b/Engine/Plugins/Enterprise/DataprepEditor/Source/DataprepCore/Private/DataprepActionAsset.cpp
@@ -11,7 +11,6 @@
 #include "DataprepParameterizableObject.h"
 #include "IDataprepLogger.h"
 #include "IDataprepProgressReporter.h"
-#include "Parameterization/DataprepParameterization.h"
 #include "SelectionSystem/DataprepFetcher.h"
 #include "SelectionSystem/DataprepFilter.h"
 #include "SelectionSystem/DataprepSelectionTransform.h"
@@ -123,7 +122,9 @@ void UDataprepActionAsset::Execute(const TArray<UObject*>& InObjects)
 			else if ( StepType == UDataprepFilter::StaticClass() )
 			{
 				UDataprepFilter* Filter = static_cast<UDataprepFilter*>( StepObject );
-				OperationContext->Context->Objects = Filter->FilterObjects( OperationContext->Context->Objects );
+
+				TArray<UObject*>& Objects = OperationContext->Context->Objects;
+				OperationContext->Context->Objects = Filter->FilterObjects( TArrayView<UObject*>( Objects.GetData(), Objects.Num() ) );
 			}
 			else if ( StepType == UDataprepSelectionTransform::StaticClass() )
 			{
@@ -467,13 +468,6 @@ bool UDataprepActionAsset::RemoveSteps(const TArray<int32>& Indices)
 		return false;
 	}
 
-	UDataprepAsset* DataprepAsset = FDataprepCoreUtils::GetDataprepAssetOfObject( this );
-	if(DataprepAsset == nullptr)
-	{
-		UE_LOG( LogDataprepCore, Error, TEXT("UDataprepActionAsset::RemoveSteps: No Dataprep asset") );
-		return false;
-	}
-
 	bool bSuccesfulRemoval = false;
 
 	Modify();
@@ -484,7 +478,7 @@ bool UDataprepActionAsset::RemoveSteps(const TArray<int32>& Indices)
 		{
 			bSuccesfulRemoval = true;
 
-			if ( UDataprepParameterization* Parameterization = DataprepAsset->GetDataprepParameterization() )
+			if ( UDataprepAsset* DataprepAsset = FDataprepCoreUtils::GetDataprepAssetOfObject( this ) )
 			{
 				TArray< UObject* > Objects;
 				GetObjectsWithOuter( Steps[Index], Objects );
@@ -492,13 +486,13 @@ bool UDataprepActionAsset::RemoveSteps(const TArray<int32>& Indices)
 				ParameterizableObjects.Reserve( Objects.Num() );
 				for ( UObject* Object : Objects )
 				{
-					if ( Object->IsA<UDataprepParameterizableObject>() )
+					if (UDataprepParameterizableObject* ParameterizableObject = Cast<UDataprepParameterizableObject>( Object ) )
 					{
-						ParameterizableObjects.Add( static_cast<UDataprepParameterizableObject*>( Object ) );
+						ParameterizableObjects.Add( ParameterizableObject );
 					}
 				}
 
-				Parameterization->RemoveBindingFromObjects( ParameterizableObjects );
+				UDataprepAsset::FRestrictedToActionAsset::NotifyAssetOfTheRemovalOfSteps( *DataprepAsset, MakeArrayView<UDataprepParameterizableObject*>( ParameterizableObjects.GetData(), ParameterizableObjects.Num() ) );
 			}
 
 			OnStepsAboutToBeRemoved.Broadcast( Steps[Index]->GetStepObject() );
@@ -535,24 +529,23 @@ void UDataprepActionAsset::NotifyDataprepSystemsOfRemoval()
 {
 	if ( UDataprepAsset* DataprepAsset = FDataprepCoreUtils::GetDataprepAssetOfObject( this ) )
 	{
-		if ( UDataprepParameterization * Parameterization = DataprepAsset->GetDataprepParameterization() )
+		
+		TArray< UObject* > Objects;
+		TArray< UDataprepParameterizableObject* > ParameterizableObjects;
+		for ( UDataprepActionStep* Step : Steps )
 		{
-			TArray< UObject* > Objects;
-			TArray< UDataprepParameterizableObject* > ParameterizableObjects;
-			for ( UDataprepActionStep* Step : Steps )
+			GetObjectsWithOuter( Step, Objects );
+			ParameterizableObjects.Reserve( Objects.Num() );
+			for ( UObject* Object : Objects )
 			{
-				GetObjectsWithOuter( Step, Objects );
-				ParameterizableObjects.Reserve( Objects.Num() );
-				for ( UObject* Object : Objects )
+				if ( UDataprepParameterizableObject* ParameterizableObject = Cast<UDataprepParameterizableObject>( Object ) )
 				{
-					if ( Object->IsA<UDataprepParameterizableObject>() )
-					{
-						ParameterizableObjects.Add( static_cast<UDataprepParameterizableObject*>( Object ) );
-					}
+					ParameterizableObjects.Add( ParameterizableObject );
 				}
 			}
-			Parameterization->RemoveBindingFromObjects( ParameterizableObjects );
 		}
+
+		UDataprepAsset::FRestrictedToActionAsset::NotifyAssetOfTheRemovalOfSteps( *DataprepAsset, MakeArrayView<UDataprepParameterizableObject*>( ParameterizableObjects.GetData(), ParameterizableObjects.Num() ) );
 	}
 }
 
@@ -572,9 +565,8 @@ void UDataprepActionAsset::ExecuteAction(const TSharedPtr<FDataprepActionContext
 			SelectedObjects.Add( Object );
 		}
 	}
-	TArray< AActor* > ActorsInWorld;
-	DataprepCorePrivateUtils::GetActorsFromWorld( ContextPtr->WorldPtr.Get(), ActorsInWorld );
-	SelectedObjects.Append( ActorsInWorld );
+
+	FDataprepCoreUtils::GetActorsFromWorld( ContextPtr->WorldPtr.Get(), SelectedObjects );
 
 	OperationContext->DataprepLogger = ContextPtr->LoggerPtr;
 	OperationContext->DataprepProgressReporter = ContextPtr->ProgressReporterPtr;
diff --git a/Engine/Plugins/Enterprise/DataprepEditor/Source/DataprepCore/Private/DataprepAsset.cpp b/Engine/Plugins/Enterprise/DataprepEditor/Source/DataprepCore/Private/DataprepAsset.cpp
--- a/Engine/Plugins/Enterprise/DataprepEditor/Source/DataprepCore/Private/DataprepAsset.cpp
+++ b/Engine/Plugins/Enterprise/DataprepEditor/Source/DataprepCore/Private/DataprepAsset.cpp
@@ -476,6 +476,16 @@ void UDataprepAsset::GetExistingParameterNamesForType(FProperty* Property, bool
 	Parameterization->GetExistingParameterNamesForType( Property, bIsDescribingFullProperty, OutValidExistingNames, OutInvalidNames );
 }
 
+void UDataprepAsset::FRestrictedToActionAsset::NotifyAssetOfTheRemovalOfSteps(UDataprepAsset& DataprepAsset, const TArrayView<UDataprepParameterizableObject*>& StepObjects)
+{
+	if ( DataprepAsset.Parameterization )
+	{
+		DataprepAsset.Parameterization->RemoveBindingFromObjects( StepObjects );
+	}
+
+	DataprepAsset.OnStepObjectsAboutToBeRemoved.Broadcast( StepObjects );
+}
+
 #ifndef NO_BLUEPRINT
 void UDataprepAsset::OnDataprepBlueprintChanged( UBlueprint* InBlueprint )
 {
diff --git a/Engine/Plugins/Enterprise/DataprepEditor/Source/DataprepCore/Private/DataprepBindingCommandChange.h b/Engine/Plugins/Enterprise/DataprepEditor/Source/DataprepCore/Private/DataprepBindingCommandChange.h
deleted file mode 100644
diff --git a/Engine/Plugins/Enterprise/DataprepEditor/Source/DataprepCore/Private/DataprepContentConsumer.cpp b/Engine/Plugins/Enterprise/DataprepEditor/Source/DataprepCore/Private/DataprepContentConsumer.cpp
--- a/Engine/Plugins/Enterprise/DataprepEditor/Source/DataprepCore/Private/DataprepContentConsumer.cpp
+++ b/Engine/Plugins/Enterprise/DataprepEditor/Source/DataprepCore/Private/DataprepContentConsumer.cpp
@@ -2,9 +2,9 @@
 
 #include "DataprepContentConsumer.h"
 
-#include "DataprepAssetUserData.h"
 #include "DataprepAssetInterface.h"
-#include "DataprepCorePrivateUtils.h"
+#include "DataprepAssetUserData.h"
+#include "DataprepCoreUtils.h"
 
 #include "EditorLevelUtils.h"
 #include "Engine/Level.h"
@@ -192,7 +192,7 @@ void UDataprepContentConsumer::AddDataprepAssetUserData()
 	}
 
 	TArray< AActor* > ActorsInWorld;
-	DataprepCorePrivateUtils::GetActorsFromWorld( Context.WorldPtr.Get(), ActorsInWorld );
+	FDataprepCoreUtils::GetActorsFromWorld( Context.WorldPtr.Get(), ActorsInWorld );
 
 	for( AActor* Actor : ActorsInWorld )
 	{
diff --git a/Engine/Plugins/Enterprise/DataprepEditor/Source/DataprepCore/Private/DataprepCorePrivateUtils.cpp b/Engine/Plugins/Enterprise/DataprepEditor/Source/DataprepCore/Private/DataprepCorePrivateUtils.cpp
--- a/Engine/Plugins/Enterprise/DataprepEditor/Source/DataprepCore/Private/DataprepCorePrivateUtils.cpp
+++ b/Engine/Plugins/Enterprise/DataprepEditor/Source/DataprepCore/Private/DataprepCorePrivateUtils.cpp
@@ -35,39 +35,6 @@ void DataprepCorePrivateUtils::DeleteRegisteredAsset(UObject* Asset)
 	}
 }
 
-void DataprepCorePrivateUtils::GetActorsFromWorld(const UWorld* World, TArray<AActor*>& OutActors )
-{
-	if(World != nullptr)
-	{
-		int32 ActorsCount = 0;
-		for(ULevel* Level : World->GetLevels())
-		{
-			ActorsCount += Level->Actors.Num();
-		}
-
-		OutActors.Reserve( OutActors.Num() + ActorsCount );
-
-		for(ULevel* Level : World->GetLevels())
-		{
-			for( AActor* Actor : Level->Actors )
-			{
-				const bool bIsValidActor = Actor
-					&& !Actor->IsPendingKill()
-					&& Actor->IsEditable()
-					&& !Actor->IsTemplate()
-					&& !FActorEditorUtils::IsABuilderBrush( Actor )
-					&& !Actor->IsA( AWorldSettings::StaticClass() )
-					&& !Actor->HasAnyFlags( RF_Transient );
-
-				if( bIsValidActor )
-				{
-					OutActors.Add( Actor );
-				}
-			}
-		}
-	}
-}
-
 
 const FString& DataprepCorePrivateUtils::GetRootTemporaryDir()
 {
diff --git a/Engine/Plugins/Enterprise/DataprepEditor/Source/DataprepCore/Private/DataprepCorePrivateUtils.h b/Engine/Plugins/Enterprise/DataprepEditor/Source/DataprepCore/Private/DataprepCorePrivateUtils.h
--- a/Engine/Plugins/Enterprise/DataprepEditor/Source/DataprepCore/Private/DataprepCorePrivateUtils.h
+++ b/Engine/Plugins/Enterprise/DataprepEditor/Source/DataprepCore/Private/DataprepCorePrivateUtils.h
@@ -51,13 +51,6 @@ namespace DataprepCorePrivateUtils
 	 */
 	void DeleteRegisteredAsset(UObject* Asset);
 
-	/**
-	 * Collect on the valid actors in the input World
-	 * @param World			World to parse
-	 * @param OutActors		Actors present in the world
-	 */
-	void GetActorsFromWorld( const UWorld* World, TArray<AActor*>& OutActors );
-
 	/** Returns directory where to store temporary files when running Dataprep asset */
 	const FString& GetRootTemporaryDir();
 
diff --git a/Engine/Plugins/Enterprise/DataprepEditor/Source/DataprepCore/Private/DataprepCoreUtils.cpp b/Engine/Plugins/Enterprise/DataprepEditor/Source/DataprepCore/Private/DataprepCoreUtils.cpp
--- a/Engine/Plugins/Enterprise/DataprepEditor/Source/DataprepCore/Private/DataprepCoreUtils.cpp
+++ b/Engine/Plugins/Enterprise/DataprepEditor/Source/DataprepCore/Private/DataprepCoreUtils.cpp
@@ -20,9 +20,11 @@
 #include "SelectionSystem/DataprepFilter.h"
 #include "SelectionSystem/DataprepSelectionTransform.h"
 
+#include "Engine/Level.h"
 #include "Engine/StaticMesh.h"
 #include "Engine/World.h"
 #include "GameFramework/Actor.h"
+#include "GameFramework/WorldSettings.h"
 #include "HAL/FileManager.h"
 #include "LevelSequence.h"
 #include "MaterialGraph/MaterialGraph.h"
@@ -33,11 +35,12 @@
 #include "Materials/MaterialInstanceConstant.h"
 #include "Misc/ScopedSlowTask.h"
 #include "RenderingThread.h"
+#include "Templates/SubclassOf.h"
 #include "UObject/StrongObjectPtr.h"
 #include "UObject/UObjectHash.h"
-#include "Templates/SubclassOf.h"
 
 #if WITH_EDITOR
+#include "ActorEditorUtils.h"
 #include "ComponentRecreateRenderStateContext.h"
 #include "Editor.h"
 #include "ObjectTools.h"
@@ -47,6 +50,44 @@
 
 #define LOCTEXT_NAMESPACE "DataprepCoreUtils"
 
+namespace DataprepCoreUtilsPrivate
+{
+	template<class ArrayOfClass>
+	void GetActorsFromWorld(const UWorld* World, TArray<ArrayOfClass*>& OutActors)
+	{
+		if ( World != nullptr )
+		{
+			int32 ActorsCount = 0;
+			for ( ULevel* Level : World->GetLevels() )
+			{
+				ActorsCount += Level->Actors.Num();
+			}
+
+			OutActors.Reserve( ActorsCount );
+
+			for ( ULevel* Level : World->GetLevels() )
+			{
+				for ( AActor* Actor : Level->Actors )
+				{
+					const bool bIsValidActor = Actor
+						&& !Actor->IsPendingKill()
+						&& Actor->IsEditable()
+						&& !Actor->IsTemplate()
+						&& !FActorEditorUtils::IsABuilderBrush( Actor )
+						&& !Actor->IsA( AWorldSettings::StaticClass() )
+						&& !Actor->HasAnyFlags( RF_Transient );
+
+					if ( bIsValidActor )
+					{
+						OutActors.Add( Actor );
+					}
+				}
+			}
+		}
+	}
+}
+
+
 UDataprepAsset* FDataprepCoreUtils::GetDataprepAssetOfObject(UObject* Object)
 {
 	while ( Object )
@@ -292,7 +333,7 @@ bool FDataprepCoreUtils::ExecuteDataprep(UDataprepAssetInterface* DataprepAssetI
 		{
 			// Delete all actors of the transient world
 			TArray<AActor*> TransientActors;
-			DataprepCorePrivateUtils::GetActorsFromWorld( TransientWorld.Get(), TransientActors );
+			FDataprepCoreUtils::GetActorsFromWorld( TransientWorld.Get(), TransientActors );
 			for ( AActor* Actor : TransientActors )
 			{
 				if ( Actor && !Actor->IsPendingKill() )
@@ -702,4 +743,15 @@ bool FDataprepCoreUtils::RemoveSteps(UDataprepActionAsset* ActionAsset, const TA
 	return ActionAsset->RemoveSteps( Indices );
 }
 
+void FDataprepCoreUtils::GetActorsFromWorld(const UWorld* World, TArray<UObject*>& OutActors)
+{
+	DataprepCoreUtilsPrivate::GetActorsFromWorld( World, OutActors );
+}
+
+void FDataprepCoreUtils::GetActorsFromWorld(const UWorld* World, TArray<AActor*>& OutActors)
+{
+	DataprepCoreUtilsPrivate::GetActorsFromWorld( World, OutActors );
+}
+
+
 #undef LOCTEXT_NAMESPACE
\ No newline at end of file
diff --git a/Engine/Plugins/Enterprise/DataprepEditor/Source/DataprepCore/Private/Parameterization/DataprepParameterization.cpp b/Engine/Plugins/Enterprise/DataprepEditor/Source/DataprepCore/Private/Parameterization/DataprepParameterization.cpp
--- a/Engine/Plugins/Enterprise/DataprepEditor/Source/DataprepCore/Private/Parameterization/DataprepParameterization.cpp
+++ b/Engine/Plugins/Enterprise/DataprepEditor/Source/DataprepCore/Private/Parameterization/DataprepParameterization.cpp
@@ -1138,7 +1138,7 @@ void UDataprepParameterization::RemoveBindedObjectProperty(UDataprepParameteriza
 	}
 }
 
-void UDataprepParameterization::RemoveBindingFromObjects(TArray<UDataprepParameterizableObject*> Objects)
+void UDataprepParameterization::RemoveBindingFromObjects(const TArrayView<UDataprepParameterizableObject*>& Objects)
 {
 	Modify();
 	TSet<FName> ParameterPotentiallyRemoved;
diff --git a/Engine/Plugins/Enterprise/DataprepEditor/Source/DataprepCore/Private/Parameterization/DataprepParameterization.h b/Engine/Plugins/Enterprise/DataprepEditor/Source/DataprepCore/Private/Parameterization/DataprepParameterization.h
--- a/Engine/Plugins/Enterprise/DataprepEditor/Source/DataprepCore/Private/Parameterization/DataprepParameterization.h
+++ b/Engine/Plugins/Enterprise/DataprepEditor/Source/DataprepCore/Private/Parameterization/DataprepParameterization.h
@@ -4,6 +4,7 @@
 
 #include "Parameterization/DataprepParameterizationUtils.h"
 
+#include "Containers/ArrayView.h"
 #include "Containers/ContainerAllocationPolicies.h"
 #include "CoreMinimal.h"
 #include "Delegates/DelegateCombinations.h"
@@ -265,7 +266,7 @@ public:
 
 	void RemoveBindedObjectProperty(UDataprepParameterizableObject* Object, const TArray<FDataprepPropertyLink>& PropertyChain);
 
-	void RemoveBindingFromObjects(TArray<UDataprepParameterizableObject*> Objects);
+	void RemoveBindingFromObjects(const TArrayView<UDataprepParameterizableObject*>& Objects);
 
 	void GetExistingParameterNamesForType(FProperty* Property, bool bIsDescribingFullProperty, TSet<FString>& OutValidExistingNames, TSet<FString>& OutInvalidNames) const;
 
diff --git a/Engine/Plugins/Enterprise/DataprepEditor/Source/DataprepCore/Private/SelectionSystem/DataprepBoolFilter.cpp b/Engine/Plugins/Enterprise/DataprepEditor/Source/DataprepCore/Private/SelectionSystem/DataprepBoolFilter.cpp
--- a/Engine/Plugins/Enterprise/DataprepEditor/Source/DataprepCore/Private/SelectionSystem/DataprepBoolFilter.cpp
+++ b/Engine/Plugins/Enterprise/DataprepEditor/Source/DataprepCore/Private/SelectionSystem/DataprepBoolFilter.cpp
@@ -3,8 +3,11 @@
 #include "SelectionSystem/DataprepBoolFilter.h"
 
 #include "DataprepCoreLogCategory.h"
-#include "SelectionSystem/DataprepSelectionSystemUtils.h"
 #include "SelectionSystem/DataprepBoolFetcher.h"
+#include "SelectionSystem/DataprepSelectionSystemStructs.h"
+#include "SelectionSystem/DataprepSelectionSystemUtils.h"
+
+#include "Containers/ArrayView.h"
 
 bool UDataprepBoolFilter::Filter(const bool bResult) const
 {
@@ -12,7 +15,7 @@ bool UDataprepBoolFilter::Filter(const bool bResult) const
 	return bResult;
 }
 
-TArray<UObject*> UDataprepBoolFilter::FilterObjects(const TArray<UObject *>& Objects) const
+TArray<UObject*> UDataprepBoolFilter::FilterObjects(const TArrayView<UObject*>& Objects) const
 {
 	if ( BoolFetcher )
 	{
@@ -23,6 +26,26 @@ TArray<UObject*> UDataprepBoolFilter::FilterObjects(const TArray<UObject *>& Obj
 	return {};
 }
 
+void UDataprepBoolFilter::FilterAndGatherInfo(const TArrayView<UObject*>& InObjects, const TArrayView<FDataprepSelectionInfo>& OutFilterResults) const
+{
+	if ( BoolFetcher )
+	{
+		return DataprepSelectionSystemUtils::FilterAndGatherInfo< UDataprepBoolFilter, UDataprepBoolFetcher, bool >( *this, *BoolFetcher, InObjects, OutFilterResults );
+	}
+
+	UE_LOG( LogDataprepCore, Error, TEXT("UDataprepBoolFilter::FilterAndGatherInfo: There was no Fetcher") );
+}
+
+void UDataprepBoolFilter::FilterAndStoreInArrayView(const TArrayView<UObject *>& InObjects, const TArrayView<bool>& OutFilterResults) const
+{
+	if ( BoolFetcher )
+	{
+		return DataprepSelectionSystemUtils::FilterAndStoreInArrayView< UDataprepBoolFilter, UDataprepBoolFetcher, bool >( *this, *BoolFetcher, InObjects, OutFilterResults );
+	}
+
+	UE_LOG( LogDataprepCore, Error, TEXT("UDataprepBoolFilter::FilterAndStoreInArrayView: There was no Fetcher") );
+}
+
 FText UDataprepBoolFilter::GetFilterCategoryText() const
 {
 	return NSLOCTEXT("DataprepBoolFilter", "BoolFilterCategory", "Condition");
diff --git a/Engine/Plugins/Enterprise/DataprepEditor/Source/DataprepCore/Private/SelectionSystem/DataprepFloatFilter.cpp b/Engine/Plugins/Enterprise/DataprepEditor/Source/DataprepCore/Private/SelectionSystem/DataprepFloatFilter.cpp
--- a/Engine/Plugins/Enterprise/DataprepEditor/Source/DataprepCore/Private/SelectionSystem/DataprepFloatFilter.cpp
+++ b/Engine/Plugins/Enterprise/DataprepEditor/Source/DataprepCore/Private/SelectionSystem/DataprepFloatFilter.cpp
@@ -27,18 +27,37 @@ bool UDataprepFloatFilter::Filter(float Float) const
 	return false;
 }
 
-TArray<UObject*> UDataprepFloatFilter::FilterObjects(const TArray<UObject*>& Objects) const
+TArray<UObject*> UDataprepFloatFilter::FilterObjects(const TArrayView<UObject*>& Objects) const
 {
 	if ( FloatFetcher )
 	{
 		return DataprepSelectionSystemUtils::FilterObjects< UDataprepFloatFilter, UDataprepFloatFetcher, float >( *this, *FloatFetcher, Objects );
 	}
 
-	ensure( false );
 	UE_LOG( LogDataprepCore, Error, TEXT("UDataprepFloatFilter::FilterObjects: There was no Fetcher") );
 	return {};
 }
 
+void UDataprepFloatFilter::FilterAndGatherInfo(const TArrayView<UObject*>& InObjects, const TArrayView<FDataprepSelectionInfo>& OutFilterResults) const
+{
+	if ( FloatFetcher )
+	{
+		return DataprepSelectionSystemUtils::FilterAndGatherInfo< UDataprepFloatFilter, UDataprepFloatFetcher, float >( *this, *FloatFetcher, InObjects, OutFilterResults );
+	}
+
+	UE_LOG( LogDataprepCore, Error, TEXT("UDataprepFloatFilter::FilterAndGatherInfo: There was no Fetcher") );
+}
+
+void UDataprepFloatFilter::FilterAndStoreInArrayView(const TArrayView<UObject*>& InObjects, const TArrayView<bool>& OutFilterResults) const
+{
+	if ( FloatFetcher )
+	{
+		return DataprepSelectionSystemUtils::FilterAndStoreInArrayView< UDataprepFloatFilter, UDataprepFloatFetcher, float >( *this, *FloatFetcher, InObjects, OutFilterResults );
+	}
+
+	UE_LOG( LogDataprepCore, Error, TEXT("UDataprepFloatFilter::FilterAndStoreInArrayView: There was no Fetcher") );
+}
+
 FText UDataprepFloatFilter::GetFilterCategoryText() const
 {
 	return NSLOCTEXT("DataprepFloatFilter", "FloatFilterCategory", "Float");
diff --git a/Engine/Plugins/Enterprise/DataprepEditor/Source/DataprepCore/Private/SelectionSystem/DataprepIntegerFilter.cpp b/Engine/Plugins/Enterprise/DataprepEditor/Source/DataprepCore/Private/SelectionSystem/DataprepIntegerFilter.cpp
--- a/Engine/Plugins/Enterprise/DataprepEditor/Source/DataprepCore/Private/SelectionSystem/DataprepIntegerFilter.cpp
+++ b/Engine/Plugins/Enterprise/DataprepEditor/Source/DataprepCore/Private/SelectionSystem/DataprepIntegerFilter.cpp
@@ -30,18 +30,37 @@ bool UDataprepIntegerFilter::Filter(int Integer) const
 	return false;
 }
 
-TArray<UObject*> UDataprepIntegerFilter::FilterObjects(const TArray<UObject*>& Objects) const
+TArray<UObject*> UDataprepIntegerFilter::FilterObjects(const TArrayView<UObject*>& Objects) const
 {
 	if ( IntFetcher )
 	{
-		return DataprepSelectionSystemUtils::FilterObjects< UDataprepIntegerFilter, UDataprepIntegerFetcher, float >( *this, *IntFetcher, Objects );
+		return DataprepSelectionSystemUtils::FilterObjects< UDataprepIntegerFilter, UDataprepIntegerFetcher, int32 >( *this, *IntFetcher, Objects );
 	}
 
-	ensure( false );
 	UE_LOG( LogDataprepCore, Error, TEXT("UDataprepIntegerFilter::FilterObjects: There was no Fetcher") );
 	return {};
 }
 
+void UDataprepIntegerFilter::FilterAndGatherInfo(const TArrayView<UObject*>& InObjects, const TArrayView<FDataprepSelectionInfo>& OutFilterResults) const
+{
+	if ( IntFetcher )
+	{
+		return DataprepSelectionSystemUtils::FilterAndGatherInfo< UDataprepIntegerFilter, UDataprepIntegerFetcher, int32 >( *this, *IntFetcher, InObjects, OutFilterResults );
+	}
+
+	UE_LOG( LogDataprepCore, Error, TEXT("UDataprepIntegerFilter::FilterAndGatherInfo: There was no Fetcher") );
+}
+
+void UDataprepIntegerFilter::FilterAndStoreInArrayView(const TArrayView<UObject*>& InObjects, const TArrayView<bool>& OutFilterResults) const
+{
+	if ( IntFetcher )
+	{
+		return DataprepSelectionSystemUtils::FilterAndStoreInArrayView< UDataprepIntegerFilter, UDataprepIntegerFetcher, int32 >( *this, *IntFetcher, InObjects, OutFilterResults );
+	}
+
+	UE_LOG( LogDataprepCore, Error, TEXT("UDataprepIntegerFilter::FilterAndStoreInArrayView: There was no Fetcher") );
+}
+
 FText UDataprepIntegerFilter::GetFilterCategoryText() const
 {
 	return NSLOCTEXT("DataprepIntegerFilter", "IntegerFilterCategory", "Integer");
diff --git a/Engine/Plugins/Enterprise/DataprepEditor/Source/DataprepCore/Private/SelectionSystem/DataprepSelectionSystemUtils.cpp b/Engine/Plugins/Enterprise/DataprepEditor/Source/DataprepCore/Private/SelectionSystem/DataprepSelectionSystemUtils.cpp
new file mode 100644
--- /dev/null
+++ b/Engine/Plugins/Enterprise/DataprepEditor/Source/DataprepCore/Private/SelectionSystem/DataprepSelectionSystemUtils.cpp
@@ -0,0 +1,33 @@
+// Copyright Epic Games, Inc. All Rights Reserved.
+
+#include "SelectionSystem/DataprepSelectionSystemUtils.h"
+
+namespace DataprepSelectionSystemUtils
+{
+	FString LexToString(const TArray<FString>& StringArray)
+	{
+		if (StringArray.Num() == 0)
+		{
+			return {};
+		}
+
+		int32 NewStringSize = 0;
+		for (const FString& Element : StringArray)
+		{
+			NewStringSize += Element.Len();
+		}
+		NewStringSize += StringArray.Num() - 1 * 2;
+
+		FString String;
+		const FString& Last = StringArray.Last();
+		for (const FString& Element : StringArray)
+		{
+			String.Append(Element);
+			if (Element != Last)
+			{
+				String.Append(TEXT(", "));
+			}
+		}
+		return String;
+	}
+}
diff --git a/Engine/Plugins/Enterprise/DataprepEditor/Source/DataprepCore/Private/SelectionSystem/DataprepSelectionSystemUtils.h b/Engine/Plugins/Enterprise/DataprepEditor/Source/DataprepCore/Private/SelectionSystem/DataprepSelectionSystemUtils.h
--- a/Engine/Plugins/Enterprise/DataprepEditor/Source/DataprepCore/Private/SelectionSystem/DataprepSelectionSystemUtils.h
+++ b/Engine/Plugins/Enterprise/DataprepEditor/Source/DataprepCore/Private/SelectionSystem/DataprepSelectionSystemUtils.h
@@ -4,8 +4,10 @@
 
 #include "SelectionSystem/DataprepFetcher.h"
 #include "SelectionSystem/DataprepFilter.h"
+#include "SelectionSystem/DataprepSelectionSystemStructs.h"
 
 #include "Async/ParallelFor.h"
+#include "Containers/ArrayView.h"
 #include "CoreMinimal.h"
 #include "Templates/UnrealTypeTraits.h"
 #include "UObject/Class.h"
@@ -13,44 +15,21 @@
 
 namespace DataprepSelectionSystemUtils
 {
-	/**
-	 * A implementation for the filtering of the objects using a filter and fetcher
-	 * This implementation is already multi-threaded when it's possible.
-	 */
-	template< class FilterClass, class FetcherClass, class FetchedDataType>
-	TArray<UObject*> FilterObjects(const FilterClass& Filter, const FetcherClass& Fetcher, const TArray<UObject*>& Objects)
-	{
-		static_assert( TIsDerivedFrom< FilterClass, UDataprepFilter >::IsDerived, "This implementation wasn't tested for a filter that isn't a child of DataprepFilter" );
-		static_assert( TIsDerivedFrom< FetcherClass, UDataprepFetcher >::IsDerived, "This implementation wasn't tested for a fetcher that isn't a child of DataprepFetcher" );
+	// Overload the function LexToString to support LexToString(const TArray<FString>&)
+	using ::LexToString;
+	FString LexToString(const TArray<FString>& StringArray);
 
-		TArray<bool> FilterResultPerObject;
-		FilterResultPerObject.AddZeroed( Objects.Num() );
-		TAtomic<int32> ObjectThatPassedCount( 0 );
-
-		auto UpdateFilterResult = [&Filter, &ObjectThatPassedCount, &FilterResultPerObject] (int32 Index, bool bFetchSucceded, const auto& FetchedData )
-		{
-			if (bFetchSucceded)
-			{
-				bool bPassedFilter = Filter.Filter( FetchedData );
-				FilterResultPerObject[Index] = bPassedFilter;
-				if ( bPassedFilter )
-				{
-					ObjectThatPassedCount++;
-				}
-			}
-			else
-			{
-				FilterResultPerObject[Index] = false;
-			}
-		};
 
+	template< class FilterClass, class FetcherClass, class FetchedDataType>
+	void DoFiltering(const FilterClass& Filter, const FetcherClass& Fetcher, const TArrayView<UObject*>& Objects, TFunctionRef<void (int32 Index, bool bFetchSucceded, const FetchedDataType& FetchedData)> FilteringFunction)
+	{
 		if ( Filter.IsThreadSafe() && Fetcher.IsThreadSafe() && bool( Fetcher.GetClass()->ClassFlags & CLASS_Native ) )
 		{
-			ParallelFor( Objects.Num(), [&Fetcher, &Objects, &UpdateFilterResult](int32 Index)
+			ParallelFor( Objects.Num(), [&Fetcher, &Objects, &FilteringFunction](int32 Index)
 			{
 				bool bFetchSucceded = false;
 				const FetchedDataType FetchedData = Fetcher.Fetch_Implementation( Objects[Index], bFetchSucceded );
-				UpdateFilterResult( Index, bFetchSucceded, FetchedData );
+				FilteringFunction( Index, bFetchSucceded, FetchedData );
 			});
 		}
 
@@ -65,10 +44,10 @@ namespace DataprepSelectionSystemUtils
 				FetcherResults.Emplace( bFetchSucceded, FetchedData );
 			}
 
-			ParallelFor( Objects.Num(), [&FetcherResults, &UpdateFilterResult](int32 Index)
+			ParallelFor( Objects.Num(), [&FetcherResults, &FilteringFunction](int32 Index)
 			{
 				const TPair< bool, FetchedDataType >& FetcherResult = FetcherResults[Index];
-				UpdateFilterResult( Index, FetcherResult.Key, FetcherResult.Value );
+				FilteringFunction( Index, FetcherResult.Key, FetcherResult.Value );
 			});
 		}
 		
@@ -87,7 +66,7 @@ namespace DataprepSelectionSystemUtils
 			for (int32 Index = 0; Index < Objects.Num(); Index++)
 			{
 				const TPair< bool, FetchedDataType >& FetcherResult = FetcherResults[Index];
-				UpdateFilterResult( Index, FetcherResult.Key, FetcherResult.Value );
+				FilteringFunction( Index, FetcherResult.Key, FetcherResult.Value );
 			}
 		}
 
@@ -97,10 +76,43 @@ namespace DataprepSelectionSystemUtils
 			{
 				bool bFetchSucceded = false;
 				const FetchedDataType FetchedData = Fetcher.Fetch( Objects[Index], bFetchSucceded );
-				UpdateFilterResult( Index, bFetchSucceded, FetchedData );
+				FilteringFunction( Index, bFetchSucceded, FetchedData );
 			}
 		}
+	}
+
+	/**
+	 * A implementation for the filtering of the objects using a filter and fetcher
+	 * This implementation is already multi-threaded when it's possible.
+	 */
+	template< class FilterClass, class FetcherClass, class FetchedDataType>
+	TArray<UObject*> FilterObjects(const FilterClass& Filter, const FetcherClass& Fetcher, const TArrayView<UObject*>& Objects)
+	{
+		static_assert( TIsDerivedFrom< FilterClass, UDataprepFilter >::IsDerived, "This implementation wasn't tested for a filter that isn't a child of DataprepFilter" );
+		static_assert( TIsDerivedFrom< FetcherClass, UDataprepFetcher >::IsDerived, "This implementation wasn't tested for a fetcher that isn't a child of DataprepFetcher" );
+
+		TArray<bool> FilterResultPerObject;
+		FilterResultPerObject.AddZeroed( Objects.Num() );
+		TAtomic<int32> ObjectThatPassedCount( 0 );
+
+		auto UpdateFilterResult = [&Filter, &ObjectThatPassedCount, &FilterResultPerObject] (int32 Index, bool bFetchSucceded, const FetchedDataType& FetchedData )
+		{
+			if (bFetchSucceded)
+			{
+				bool bPassedFilter = Filter.Filter( FetchedData );
+				FilterResultPerObject[Index] = bPassedFilter;
+				if ( bPassedFilter )
+				{
+					ObjectThatPassedCount++;
+				}
+			}
+			else
+			{
+				FilterResultPerObject[Index] = false;
+			}
+		};
 
+		DoFiltering<FilterClass, FetcherClass, FetchedDataType>( Filter, Fetcher, Objects, UpdateFilterResult );
 	
 		const bool bIsExcludingObjectThatPassedFilter = Filter.IsExcludingResult();
 
@@ -118,5 +130,85 @@ namespace DataprepSelectionSystemUtils
 
 		return SelectedObjects;
 	}
-}
 
+	template<class FetchedDataType>
+	static typename TEnableIf<TIsFilterDataDisplayable<FetchedDataType>::Value>::Type UpdateReportedDataIfSupported(const FetchedDataType& FetchedData, FDataprepSelectionInfo& OutInfo)
+	{
+		OutInfo.FetchedData.Set<FetchedDataType>( FetchedData );
+		OutInfo.bWasDataFetchedAndCached = true;
+	}
+
+	template<class FetchedDataType>
+	static typename TEnableIf<!TIsFilterDataDisplayable<FetchedDataType>::Value>::Type UpdateReportedDataIfSupported(const FetchedDataType& FetchedData, FDataprepSelectionInfo& OutInfo)
+	{
+		// Do nothing since the data type is not supported
+		OutInfo.bWasDataFetchedAndCached = false;
+	}
+
+	/**
+	 * A implementation for the filtering of the objects using a filter and fetcher
+	 * This one also gather additional info on how the filtering was done
+	 * This implementation is already multi-threaded when it's possible.
+	 */
+	template< class FilterClass, class FetcherClass, class FetchedDataType>
+	void FilterAndGatherInfo(const FilterClass& Filter, const FetcherClass& Fetcher, const TArrayView<UObject*>& InObjects, const TArrayView<FDataprepSelectionInfo>& OutFilterResults)
+	{
+		static_assert( TIsDerivedFrom< FilterClass, UDataprepFilter >::IsDerived, "This implementation wasn't tested for a filter that isn't a child of DataprepFilter" );
+		static_assert( TIsDerivedFrom< FetcherClass, UDataprepFetcher >::IsDerived, "This implementation wasn't tested for a fetcher that isn't a child of DataprepFetcher" );
+
+		check( InObjects.Num() == OutFilterResults.Num() );
+
+		auto UpdateFilterResult = [&Filter, &OutFilterResults] (int32 Index, bool bFetchSucceded, const FetchedDataType& FetchedData )
+		{
+			FDataprepSelectionInfo& SelectionInfo = OutFilterResults[ Index ];
+			SelectionInfo.bHasPassFilter = false;
+
+			if ( bFetchSucceded )
+			{
+				SelectionInfo.bHasPassFilter = Filter.Filter( FetchedData );
+				UpdateReportedDataIfSupported<FetchedDataType>( FetchedData, SelectionInfo );
+			}
+			else
+			{
+				SelectionInfo.bWasDataFetchedAndCached = false;
+			}
+
+			if ( Filter.IsExcludingResult() )
+			{
+				SelectionInfo.bHasPassFilter = !SelectionInfo.bHasPassFilter;
+			}
+		};
+
+		DoFiltering<FilterClass, FetcherClass, FetchedDataType>( Filter, Fetcher, InObjects, UpdateFilterResult );
+	}
+
+	/**
+	 * A implementation for the filtering of the objects using a filter and fetcher
+	 * This implementation is already multi-threaded when it's possible.
+	 */
+	template< class FilterClass, class FetcherClass, class FetchedDataType>
+	void FilterAndStoreInArrayView(const FilterClass& Filter, const FetcherClass& Fetcher, const TArrayView<UObject*>& InObjects, const TArrayView<bool>& OutFilterResults)
+	{
+		static_assert( TIsDerivedFrom< FilterClass, UDataprepFilter >::IsDerived, "This implementation wasn't tested for a filter that isn't a child of DataprepFilter" );
+		static_assert( TIsDerivedFrom< FetcherClass, UDataprepFetcher >::IsDerived, "This implementation wasn't tested for a fetcher that isn't a child of DataprepFetcher" );
+
+		check( InObjects.Num() == OutFilterResults.Num() );
+
+		auto UpdateFilterResult = [&Filter, &OutFilterResults] (int32 Index, bool bFetchSucceded, const FetchedDataType& FetchedData )
+		{
+			bool& bHasPass = OutFilterResults[ Index ];
+			bHasPass = false;
+			if ( bFetchSucceded )
+			{
+				bHasPass = Filter.Filter( FetchedData );
+			}
+
+			if ( Filter.IsExcludingResult() )
+			{
+				bHasPass = !bHasPass;
+			}
+		};
+
+		DoFiltering<FilterClass, FetcherClass, FetchedDataType>( Filter, Fetcher, InObjects, UpdateFilterResult );
+	}
+}
diff --git a/Engine/Plugins/Enterprise/DataprepEditor/Source/DataprepCore/Private/SelectionSystem/DataprepStringFilter.cpp b/Engine/Plugins/Enterprise/DataprepEditor/Source/DataprepCore/Private/SelectionSystem/DataprepStringFilter.cpp
--- a/Engine/Plugins/Enterprise/DataprepEditor/Source/DataprepCore/Private/SelectionSystem/DataprepStringFilter.cpp
+++ b/Engine/Plugins/Enterprise/DataprepEditor/Source/DataprepCore/Private/SelectionSystem/DataprepStringFilter.cpp
@@ -27,18 +27,37 @@ bool UDataprepStringFilter::Filter(const FString& String) const
 	return false;
 }
 
-TArray<UObject*> UDataprepStringFilter::FilterObjects(const TArray<UObject *>& Objects) const
+TArray<UObject*> UDataprepStringFilter::FilterObjects(const TArrayView<UObject*>& Objects) const
 {
 	if ( StringFetcher )
 	{
 		return DataprepSelectionSystemUtils::FilterObjects< UDataprepStringFilter, UDataprepStringFetcher, FString >( *this, *StringFetcher, Objects );
 	}
 
-	ensure(false);
-	UE_LOG(LogDataprepCore, Error, TEXT("UDataprepStringFilter::FilterObjects: There was no Fetcher"));
+	UE_LOG( LogDataprepCore, Error, TEXT("UDataprepStringFilter::FilterObjects: There was no Fetcher") );
 	return {};
 }
 
+void UDataprepStringFilter::FilterAndGatherInfo(const TArrayView<UObject*>& InObjects, const TArrayView<FDataprepSelectionInfo>& OutFilterResults) const
+{
+	if ( StringFetcher )
+	{
+		return DataprepSelectionSystemUtils::FilterAndGatherInfo< UDataprepStringFilter, UDataprepStringFetcher, FString >( *this, *StringFetcher, InObjects, OutFilterResults );
+	}
+
+	UE_LOG( LogDataprepCore, Error, TEXT("UDataprepStringFilter::FilterAndGatherInfo: There was no Fetcher") );
+}
+
+void UDataprepStringFilter::FilterAndStoreInArrayView(const TArrayView<UObject*>& InObjects, const TArrayView<bool>& OutFilterResults) const
+{
+	if ( StringFetcher )
+	{
+		return DataprepSelectionSystemUtils::FilterAndStoreInArrayView< UDataprepStringFilter, UDataprepStringFetcher, FString >( *this, *StringFetcher, InObjects, OutFilterResults );
+	}
+
+	UE_LOG( LogDataprepCore, Error, TEXT("UDataprepStringFilter::FilterAndStoreInArrayView: There was no Fetcher") );
+}
+
 TSubclassOf<UDataprepFetcher> UDataprepStringFilter::GetAcceptedFetcherClass() const
 {
 	return UDataprepStringFetcher::StaticClass();
diff --git a/Engine/Plugins/Enterprise/DataprepEditor/Source/DataprepCore/Private/SelectionSystem/DataprepStringsArrayFilter.cpp b/Engine/Plugins/Enterprise/DataprepEditor/Source/DataprepCore/Private/SelectionSystem/DataprepStringsArrayFilter.cpp
--- a/Engine/Plugins/Enterprise/DataprepEditor/Source/DataprepCore/Private/SelectionSystem/DataprepStringsArrayFilter.cpp
+++ b/Engine/Plugins/Enterprise/DataprepEditor/Source/DataprepCore/Private/SelectionSystem/DataprepStringsArrayFilter.cpp
@@ -34,18 +34,37 @@ bool UDataprepStringsArrayFilter::Filter(const TArray<FString>& StringArray) con
 	return bResult;
 }
 
-TArray<UObject*> UDataprepStringsArrayFilter::FilterObjects(const TArray<UObject *>& Objects) const
+TArray<UObject*> UDataprepStringsArrayFilter::FilterObjects(const TArrayView<UObject*>& Objects) const
 {
 	if ( StringsArrayFetcher )
 	{
 		return DataprepSelectionSystemUtils::FilterObjects< UDataprepStringsArrayFilter, UDataprepStringsArrayFetcher, TArray<FString> >( *this, *StringsArrayFetcher, Objects );
 	}
 
-	ensure(false);
 	UE_LOG(LogDataprepCore, Error, TEXT("UDataprepStringsArrayFilter::FilterObjects: There was no Fetcher"));
 	return {};
 }
 
+void UDataprepStringsArrayFilter::FilterAndGatherInfo(const TArrayView<UObject *>& InObjects, const TArrayView<FDataprepSelectionInfo>& OutFilterResults) const
+{
+	if ( StringsArrayFetcher )
+	{
+		return DataprepSelectionSystemUtils::FilterAndGatherInfo< UDataprepStringsArrayFilter, UDataprepStringsArrayFetcher, TArray<FString> >( *this, *StringsArrayFetcher, InObjects, OutFilterResults );
+	}
+
+	UE_LOG( LogDataprepCore, Error, TEXT("UDataprepStringsArrayFilter::FilterAndGatherInfo: There was no Fetcher") );
+}
+
+void UDataprepStringsArrayFilter::FilterAndStoreInArrayView(const TArrayView<UObject*>& InObjects, const TArrayView<bool>& OutFilterResults) const
+{
+	if ( StringsArrayFetcher )
+	{
+		return DataprepSelectionSystemUtils::FilterAndStoreInArrayView< UDataprepStringsArrayFilter, UDataprepStringsArrayFetcher, TArray<FString> >( *this, *StringsArrayFetcher, InObjects, OutFilterResults );
+	}
+
+	UE_LOG( LogDataprepCore, Error, TEXT("UDataprepStringsArrayFilter::FilterAndStoreInArrayView: There was no Fetcher") );
+}
+
 TSubclassOf<UDataprepFetcher> UDataprepStringsArrayFilter::GetAcceptedFetcherClass() const
 {
 	return UDataprepStringsArrayFetcher::StaticClass();
diff --git a/Engine/Plugins/Enterprise/DataprepEditor/Source/DataprepCore/Public/DataprepAsset.h b/Engine/Plugins/Enterprise/DataprepEditor/Source/DataprepCore/Public/DataprepAsset.h
--- a/Engine/Plugins/Enterprise/DataprepEditor/Source/DataprepCore/Public/DataprepAsset.h
+++ b/Engine/Plugins/Enterprise/DataprepEditor/Source/DataprepCore/Public/DataprepAsset.h
@@ -149,6 +149,19 @@ public:
 	DECLARE_EVENT_TwoParams(UDataprepAsset, FOnDataprepActionAssetChange, UObject* /*The object that was modified*/, FDataprepAssetChangeType)
 	FOnDataprepActionAssetChange& GetOnActionChanged() { return OnActionChanged; }
 
+	/**
+	 * Allow an observer to be notified of the removal of some step from the asset
+	 */
+	DECLARE_EVENT_OneParam(UDataprepAsset, FOnStepObjectsAboutToBeRemoved, const TArrayView<UDataprepParameterizableObject*>&)
+	FOnStepObjectsAboutToBeRemoved& GetOnStepObjectsAboutToBeRemoved() { return OnStepObjectsAboutToBeRemoved; }
+
+	struct FRestrictedToActionAsset
+	{
+	private:
+		friend UDataprepActionAsset;
+		static void NotifyAssetOfTheRemovalOfSteps(UDataprepAsset& DataprepAsset, const TArrayView<UDataprepParameterizableObject*>& StepObjects);
+	};
+
 	bool CreateParameterization();
 
 #ifndef NO_BLUEPRINT
@@ -240,6 +253,7 @@ private:
 
 	FOnDataprepActionAssetChange OnActionChanged;
 
-	int32 CachedActionCount;
+	FOnStepObjectsAboutToBeRemoved OnStepObjectsAboutToBeRemoved;
 
+	int32 CachedActionCount;
 };
\ No newline at end of file
diff --git a/Engine/Plugins/Enterprise/DataprepEditor/Source/DataprepCore/Public/DataprepBindingCommandChange.h b/Engine/Plugins/Enterprise/DataprepEditor/Source/DataprepCore/Public/DataprepBindingCommandChange.h
new file mode 100644
--- /dev/null
+++ b/Engine/Plugins/Enterprise/DataprepEditor/Source/DataprepCore/Public/DataprepBindingCommandChange.h
@@ -0,0 +1,100 @@
+// Copyright Epic Games, Inc. All Rights Reserved.
+
+#pragma once
+
+#include "Misc/Change.h"
+
+template<class ObservedObjectType, typename DelatgateToBindType>
+class TDataprepBindingCommandChange final : public FCommandChange
+{
+public:
+	DECLARE_DELEGATE_RetVal(DelatgateToBindType&, FDelegateAccessor);
+	DECLARE_DELEGATE_TwoParams(OnUndoRedoDelegateType, ObservedObjectType&, const FDelegateHandle&);
+
+	using FDelegate = typename DelatgateToBindType::FDelegate;
+
+	TDataprepBindingCommandChange(bool bInWasBindingAdded
+		, ObservedObjectType& ObservedObject
+		, FDelegate&& InBinding
+		, FDelegateAccessor InDelegateAccessor
+		, OnUndoRedoDelegateType&& InOnApply
+		, OnUndoRedoDelegateType&& InOnRevert
+		, const FDelegateHandle& InHandle
+		);
+
+	virtual void Apply(UObject* Object) override;
+
+	virtual void Revert(UObject* Object) override;
+
+	virtual FString ToString() const override
+	{
+		return TEXT("Binding Tracking");
+	}
+
+private:
+	bool bWasBindingAdded;
+	TWeakObjectPtr<ObservedObjectType> ObservedObjectWeakPtr;
+	FDelegate Binding;
+	FDelegateAccessor DelegateAccessor;
+	OnUndoRedoDelegateType OnApply;
+	OnUndoRedoDelegateType OnRevert;
+	FDelegateHandle Handle;
+};
+
+template<class ObservedObjectType, typename DelatgateToBindType>
+TDataprepBindingCommandChange<ObservedObjectType, DelatgateToBindType>::TDataprepBindingCommandChange(
+	bool bInWasBindingAdded
+	, ObservedObjectType& ObservedObject
+	, FDelegate&& InBinding
+	, FDelegateAccessor InDelegateAccessor
+	, OnUndoRedoDelegateType&& InOnApply
+	, OnUndoRedoDelegateType&& InOnRevert
+	, const FDelegateHandle& InHandle
+	)
+	: bWasBindingAdded( bInWasBindingAdded )
+	, ObservedObjectWeakPtr( &ObservedObject )
+	, Binding( MoveTemp( InBinding ) )
+	, DelegateAccessor( MoveTemp( InDelegateAccessor ) )
+	, OnApply( MoveTemp( InOnApply ) )
+	, OnRevert( MoveTemp( InOnRevert ) )
+	, Handle( InHandle )
+{
+}
+
+template<class ObservedObjectType, typename DelatgateToBindType>
+void TDataprepBindingCommandChange<ObservedObjectType, DelatgateToBindType>::Apply(UObject* Object)
+{
+	ObservedObjectType* ObservedObject = ObservedObjectWeakPtr.Get();
+	if ( ObservedObject && DelegateAccessor.IsBound() )
+	{
+		DelatgateToBindType& Delegate = DelegateAccessor.Execute();
+		if ( bWasBindingAdded )
+		{
+			Handle = Delegate.Add( Binding );
+		}
+		else
+		{
+			Delegate.Remove( Handle );
+		}
+		OnApply.ExecuteIfBound( *ObservedObject, Handle );
+	}
+}
+
+template<class ObservedObjectType, typename DelatgateToBindType>
+void TDataprepBindingCommandChange<ObservedObjectType, DelatgateToBindType>::Revert(UObject* Object)
+{
+	ObservedObjectType* ObservedObject = ObservedObjectWeakPtr.Get();
+	if ( ObservedObject && DelegateAccessor.IsBound() )
+	{
+		DelatgateToBindType& Delegate = DelegateAccessor.Execute();
+		if ( bWasBindingAdded )
+		{
+			Delegate.Remove( Handle );
+		}
+		else
+		{
+			Handle = Delegate.Add( Binding );
+		}
+		OnRevert.ExecuteIfBound( *ObservedObject, Handle );
+	}
+}
diff --git a/Engine/Plugins/Enterprise/DataprepEditor/Source/DataprepCore/Public/DataprepCoreUtils.h b/Engine/Plugins/Enterprise/DataprepEditor/Source/DataprepCore/Public/DataprepCoreUtils.h
--- a/Engine/Plugins/Enterprise/DataprepEditor/Source/DataprepCore/Public/DataprepCoreUtils.h
+++ b/Engine/Plugins/Enterprise/DataprepEditor/Source/DataprepCore/Public/DataprepCoreUtils.h
@@ -221,5 +221,20 @@ public:
 		int32 TaskDepth;
 		TUniquePtr<FFeedbackContextAnsi> FeedbackContext;
 	};
+
+	/**
+	 * Collect on the valid actors in the input World
+	 * @param World			World to parse
+	 * @param OutActors		Actors present in the world
+	 */
+	static void GetActorsFromWorld(const UWorld* World, TArray<UObject*>& OutActors);
+
+	/**
+	 * Collect on the valid actors in the input World
+	 * @param World			World to parse
+	 * @param OutActors		Actors present in the world
+	 */
+	static void GetActorsFromWorld(const UWorld* World, TArray<AActor*>& OutActors);
+
 };
 
diff --git a/Engine/Plugins/Enterprise/DataprepEditor/Source/DataprepCore/Public/SelectionSystem/DataprepBoolFilter.h b/Engine/Plugins/Enterprise/DataprepEditor/Source/DataprepCore/Public/SelectionSystem/DataprepBoolFilter.h
--- a/Engine/Plugins/Enterprise/DataprepEditor/Source/DataprepCore/Public/SelectionSystem/DataprepBoolFilter.h
+++ b/Engine/Plugins/Enterprise/DataprepEditor/Source/DataprepCore/Public/SelectionSystem/DataprepBoolFilter.h
@@ -21,7 +21,9 @@ public:
 	bool Filter(const bool bResult) const;
 
 	//~ Begin UDataprepFilter Interface
-	virtual TArray<UObject*> FilterObjects(const TArray<UObject*>& Objects) const override;
+	virtual TArray<UObject*> FilterObjects(const TArrayView<UObject*>& Objects) const override;
+	virtual void FilterAndGatherInfo(const TArrayView<UObject*>& InObjects, const TArrayView<FDataprepSelectionInfo>& OutFilterResults) const override;
+	virtual void FilterAndStoreInArrayView(const TArrayView<UObject *>& InObjects, const TArrayView<bool>& OutFilterResults) const override;
 	virtual bool IsThreadSafe() const override { return true; }
 	virtual FText GetFilterCategoryText() const override;
 	virtual TSubclassOf<UDataprepFetcher> GetAcceptedFetcherClass() const override;
diff --git a/Engine/Plugins/Enterprise/DataprepEditor/Source/DataprepCore/Public/SelectionSystem/DataprepFilter.h b/Engine/Plugins/Enterprise/DataprepEditor/Source/DataprepCore/Public/SelectionSystem/DataprepFilter.h
--- a/Engine/Plugins/Enterprise/DataprepEditor/Source/DataprepCore/Public/SelectionSystem/DataprepFilter.h
+++ b/Engine/Plugins/Enterprise/DataprepEditor/Source/DataprepCore/Public/SelectionSystem/DataprepFilter.h
@@ -6,6 +6,7 @@
 
 #include "SelectionSystem/DataprepFetcher.h"
 
+#include "Containers/ArrayView.h"
 #include "CoreMinimal.h"
 #include "Templates/SubclassOf.h"
 #include "UObject/Class.h"
@@ -14,6 +15,8 @@
 
 #include "DataprepFilter.generated.h"
 
+struct FDataprepSelectionInfo;
+
 /**
  * The Dataprep Filter a base class for the Dataprep selection system
  * It's main responsibility is to filter a array of object and to return the selected objects
@@ -30,11 +33,27 @@ public:
 	// End UObject Interface
 
 	/**
-	 * Take a array of object and return the objects that pass the filter
+	 * Take an array of objects and return the objects that pass the filter
 	 * @param Objects The object to filter
 	 * @return The object that passed the filtering
 	 */
-	virtual TArray<UObject*> FilterObjects(const TArray<UObject*>& Objects) const{ return {}; }
+	virtual TArray<UObject*> FilterObjects(const TArrayView<UObject*>& Objects) const { return {}; }
+
+	/**
+	 * Take an array of object and output the result into the arrays
+	 * @param InObjects The object to filter
+	 * @param OutFilterResult Will put true at the same index of the object if it passed the filter
+	 * @param OutFetchedSucces Will put true at the same index of the object if the fetcher fetched a value
+	 * @param OutFetchedLabel Will put a display of information contextual to the filter/fetcher
+	 */
+	virtual void FilterAndGatherInfo(const TArrayView<UObject*>& InObjects, const TArrayView<FDataprepSelectionInfo>& OutFilterResults) const {};
+
+	/**
+	 * Take an array of object and output the result into the result array
+	 * @param InObjects The object to filter
+	 * @param OutFilterResult Will put true at the same index of the object if it has passed the filter
+	 */
+	virtual void FilterAndStoreInArrayView(const TArrayView<UObject*>& InObjects, const TArrayView<bool>& OutFilterResults) const {};
 
 	/**
 	 * Is this filter safe to use in a multi thread execution?
@@ -83,6 +102,14 @@ public:
 	{
 		Modify();
 		bIsExcludingResult = bInIsExcludingResult;
+		FProperty* Property = StaticClass()->FindPropertyByName( GET_MEMBER_NAME_CHECKED( UDataprepFilter, bIsExcludingResult ) );
+		check(Property);
+		FEditPropertyChain EditChain;
+		EditChain.AddHead( Property );
+		EditChain.SetActivePropertyNode( Property );
+		FPropertyChangedEvent EditPropertyChangeEvent( Property, EPropertyChangeType::ValueSet );
+		FPropertyChangedChainEvent EditChangeChainEvent( EditChain, EditPropertyChangeEvent );
+		PostEditChangeChainProperty( EditChangeChainEvent );
 	}
 
 	/**
diff --git a/Engine/Plugins/Enterprise/DataprepEditor/Source/DataprepCore/Public/SelectionSystem/DataprepFloatFilter.h b/Engine/Plugins/Enterprise/DataprepEditor/Source/DataprepCore/Public/SelectionSystem/DataprepFloatFilter.h
--- a/Engine/Plugins/Enterprise/DataprepEditor/Source/DataprepCore/Public/SelectionSystem/DataprepFloatFilter.h
+++ b/Engine/Plugins/Enterprise/DataprepEditor/Source/DataprepCore/Public/SelectionSystem/DataprepFloatFilter.h
@@ -29,7 +29,9 @@ public:
 	bool Filter(float Float) const;
 
 	// Begin UDataprepFilter Interface
-	virtual TArray<UObject*> FilterObjects(const TArray<UObject*>& Objects) const override;
+	virtual TArray<UObject*> FilterObjects(const TArrayView<UObject*>& Objects) const override;
+	virtual void FilterAndGatherInfo(const TArrayView<UObject*>& InObjects, const TArrayView<FDataprepSelectionInfo>& OutFilterResults) const override;
+	virtual void FilterAndStoreInArrayView(const TArrayView<UObject*>& InObjects, const TArrayView<bool>& OutFilterResults) const override;
 	virtual bool IsThreadSafe() const override { return true; }
 	virtual FText GetFilterCategoryText() const override;
 	virtual TSubclassOf<UDataprepFetcher> GetAcceptedFetcherClass() const override;
diff --git a/Engine/Plugins/Enterprise/DataprepEditor/Source/DataprepCore/Public/SelectionSystem/DataprepIntegerFetcher.h b/Engine/Plugins/Enterprise/DataprepEditor/Source/DataprepCore/Public/SelectionSystem/DataprepIntegerFetcher.h
--- a/Engine/Plugins/Enterprise/DataprepEditor/Source/DataprepCore/Public/SelectionSystem/DataprepIntegerFetcher.h
+++ b/Engine/Plugins/Enterprise/DataprepEditor/Source/DataprepCore/Public/SelectionSystem/DataprepIntegerFetcher.h
@@ -27,7 +27,7 @@ public:
 	 * @return The fetched float
 	 */
 	UFUNCTION(BlueprintNativeEvent)
-	int Fetch(const UObject* Object, bool& bOutFetchSucceded) const;
+	int32 Fetch(const UObject* Object, bool& bOutFetchSucceded) const;
 
 	/**
 	 * This function is the same has Fetch, but it's the extension point for an operation defined in c++.
@@ -36,7 +36,7 @@ public:
 	 * @param bOutFetchSucceded If the fetcher managed to retrieve a float from the object this bool must be set to true
 	 * @return The fetched float
 	 */
-	virtual int Fetch_Implementation(const UObject* Object, bool& bOutFetchSucceded) const
+	virtual int32 Fetch_Implementation(const UObject* Object, bool& bOutFetchSucceded) const
 	{
 		bOutFetchSucceded = false;
 		return {};
diff --git a/Engine/Plugins/Enterprise/DataprepEditor/Source/DataprepCore/Public/SelectionSystem/DataprepIntegerFilter.h b/Engine/Plugins/Enterprise/DataprepEditor/Source/DataprepCore/Public/SelectionSystem/DataprepIntegerFilter.h
--- a/Engine/Plugins/Enterprise/DataprepEditor/Source/DataprepCore/Public/SelectionSystem/DataprepIntegerFilter.h
+++ b/Engine/Plugins/Enterprise/DataprepEditor/Source/DataprepCore/Public/SelectionSystem/DataprepIntegerFilter.h
@@ -30,7 +30,9 @@ public:
 	bool Filter(int Integer) const;
 
 	//~ Begin UDataprepFilter Interface
-	virtual TArray<UObject*> FilterObjects(const TArray<UObject*>& Objects) const override;
+	virtual TArray<UObject*> FilterObjects(const TArrayView<UObject*>& Objects) const override;
+	virtual void FilterAndGatherInfo(const TArrayView<UObject*>& InObjects, const TArrayView<FDataprepSelectionInfo>& OutFilterResults) const override;
+	virtual void FilterAndStoreInArrayView(const TArrayView<UObject*>& InObjects, const TArrayView<bool>& OutFilterResults) const override;
 	virtual bool IsThreadSafe() const override { return true; }
 	virtual FText GetFilterCategoryText() const override;
 	virtual TSubclassOf<UDataprepFetcher> GetAcceptedFetcherClass() const override;
diff --git a/Engine/Plugins/Enterprise/DataprepEditor/Source/DataprepCore/Public/SelectionSystem/DataprepSelectionSystemStructs.h b/Engine/Plugins/Enterprise/DataprepEditor/Source/DataprepCore/Public/SelectionSystem/DataprepSelectionSystemStructs.h
new file mode 100644
--- /dev/null
+++ b/Engine/Plugins/Enterprise/DataprepEditor/Source/DataprepCore/Public/SelectionSystem/DataprepSelectionSystemStructs.h
@@ -0,0 +1,50 @@
+// Copyright Epic Games, Inc. All Rights Reserved.
+
+#pragma once
+
+#include "CoreMinimal.h"
+#include "Misc/TVariant.h"
+
+// Those template must be keep in sync with the variant. Those are the data types that the filter preview can display the data in the columns (scene preview and asset preview)
+template <class T>
+struct TIsFilterDataDisplayable
+{
+	enum { Value = false };
+};
+
+template <>
+struct TIsFilterDataDisplayable<int32>
+{
+	enum { Value = true };
+};
+
+template <>
+struct TIsFilterDataDisplayable<float>
+{
+	enum { Value = true };
+};
+
+template <>
+struct TIsFilterDataDisplayable<FString>
+{
+	enum { Value = true };
+};
+
+/**
+ * This type is the variant used by the dataprep filter preview system to store a fetched data and then displaying it latter.
+ * if you update this, add the new type to the list above. 
+ * Try to compile it. It should fail on a couple of functors. 
+ * Update those and keep the functions explicit so that the compilation will fail when adding a type to the variant.
+ * This way the programmer will be force to acknowledge how those functors should behave on the new type
+ */
+using FFilterVariantData = TVariant< FEmptyVariantState, int32, float, FString >;
+
+/**
+ * A struct to cache the result of a filter with some additional info
+ */
+struct DATAPREPCORE_API FDataprepSelectionInfo
+{
+	FFilterVariantData FetchedData;
+	uint8 bHasPassFilter:1;
+	uint8 bWasDataFetchedAndCached:1;
+};
diff --git a/Engine/Plugins/Enterprise/DataprepEditor/Source/DataprepCore/Public/SelectionSystem/DataprepStringFilter.h b/Engine/Plugins/Enterprise/DataprepEditor/Source/DataprepCore/Public/SelectionSystem/DataprepStringFilter.h
--- a/Engine/Plugins/Enterprise/DataprepEditor/Source/DataprepCore/Public/SelectionSystem/DataprepStringFilter.h
+++ b/Engine/Plugins/Enterprise/DataprepEditor/Source/DataprepCore/Public/SelectionSystem/DataprepStringFilter.h
@@ -30,7 +30,9 @@ public:
 	bool Filter(const FString& String) const;
 
 	//~ Begin UDataprepFilter Interface
-	virtual TArray<UObject*> FilterObjects(const TArray<UObject*>& Objects) const override;
+	virtual TArray<UObject*> FilterObjects(const TArrayView<UObject*>& Objects) const override;
+	virtual void FilterAndGatherInfo(const TArrayView<UObject*>& InObjects, const TArrayView<FDataprepSelectionInfo>& OutFilterResults) const override;
+	virtual void FilterAndStoreInArrayView(const TArrayView<UObject*>& InObjects, const TArrayView<bool>& OutFilterResults) const override;
 	virtual bool IsThreadSafe() const override { return true; }
 	virtual FText GetFilterCategoryText() const override;
 	virtual TSubclassOf<UDataprepFetcher> GetAcceptedFetcherClass() const override;
diff --git a/Engine/Plugins/Enterprise/DataprepEditor/Source/DataprepCore/Public/SelectionSystem/DataprepStringsArrayFilter.h b/Engine/Plugins/Enterprise/DataprepEditor/Source/DataprepCore/Public/SelectionSystem/DataprepStringsArrayFilter.h
--- a/Engine/Plugins/Enterprise/DataprepEditor/Source/DataprepCore/Public/SelectionSystem/DataprepStringsArrayFilter.h
+++ b/Engine/Plugins/Enterprise/DataprepEditor/Source/DataprepCore/Public/SelectionSystem/DataprepStringsArrayFilter.h
@@ -17,7 +17,9 @@ public:
 	bool Filter(const TArray<FString>& StringArray) const;
 
 	//~ Begin UDataprepFilter Interface
-	virtual TArray<UObject*> FilterObjects(const TArray<UObject*>& Objects) const override;
+	virtual TArray<UObject*> FilterObjects(const TArrayView<UObject*>& Objects) const override;
+	virtual void FilterAndGatherInfo(const TArrayView<UObject*>& InObjects, const TArrayView<FDataprepSelectionInfo>& OutFilterResults) const override;
+	virtual void FilterAndStoreInArrayView(const TArrayView<UObject*>& InObjects, const TArrayView<bool>& OutFilterResults) const override;
 	virtual bool IsThreadSafe() const override { return true; }
 	virtual FText GetFilterCategoryText() const override;
 	virtual TSubclassOf<UDataprepFetcher> GetAcceptedFetcherClass() const override;
diff --git a/Engine/Plugins/Enterprise/DataprepEditor/Source/DataprepEditor/Private/DataprepEditor.cpp b/Engine/Plugins/Enterprise/DataprepEditor/Source/DataprepEditor/Private/DataprepEditor.cpp
--- a/Engine/Plugins/Enterprise/DataprepEditor/Source/DataprepEditor/Private/DataprepEditor.cpp
+++ b/Engine/Plugins/Enterprise/DataprepEditor/Source/DataprepEditor/Private/DataprepEditor.cpp
@@ -4,18 +4,22 @@
 
 #include "DataprepOperation.h"
 
+#include "DataprepActionAsset.h"
 #include "DataprepAssetInstance.h"
 #include "DataprepContentConsumer.h"
 #include "DataprepContentProducer.h"
+#include "DataprepCoreUtils.h"
 #include "DataprepEditorActions.h"
+#include "DataprepEditorLogCategory.h"
 #include "DataprepEditorModule.h"
 #include "DataprepEditorStyle.h"
 #include "DataprepRecipe.h"
-#include "DataprepActionAsset.h"
-#include "DataprepCoreUtils.h"
-#include "DataprepEditorLogCategory.h"
+#include "PreviewSystem/DataprepPreviewAssetColumn.h"
+#include "PreviewSystem/DataprepPreviewSceneOutlinerColumn.h"
+#include "PreviewSystem/DataprepPreviewSystem.h"
 #include "SchemaActions/DataprepOperationMenuActionCollector.h"
 #include "Widgets/DataprepAssetView.h"
+#include "Widgets/DataprepGraph/SDataprepGraphEditor.h"
 #include "Widgets/SAssetsPreviewWidget.h"
 #include "Widgets/SDataprepEditorViewport.h"
 #include "Widgets/SDataprepPalette.h"
@@ -24,7 +28,6 @@
 #include "ActorTreeItem.h"
 #include "AssetDeleteModel.h"
 #include "AssetRegistryModule.h"
-#include "StatsViewerModule.h"
 #include "BlueprintNodeSpawner.h"
 #include "ComponentTreeItem.h"
 #include "DesktopPlatformModule.h"
@@ -50,7 +53,9 @@
 #include "Misc/ScopedSlowTask.h"
 #include "Modules/ModuleManager.h"
 #include "ObjectTools.h"
+#include "SceneOutlinerPublicTypes.h"
 #include "ScopedTransaction.h"
+#include "StatsViewerModule.h"
 #include "Templates/UnrealTemplate.h"
 #include "Toolkits/IToolkit.h"
 #include "UObject/Object.h"
@@ -152,6 +157,7 @@ FDataprepEditor::FDataprepEditor()
 	, bSaveIntermediateBuildProducts(false)
 	, PreviewWorld(nullptr)
 	, bIgnoreCloseRequest(false)
+	, PreviewSystem( MakeShared<FDataprepPreviewSystem>() )
 {
 	FName UniqueWorldName = MakeUniqueObjectName(GetTransientPackage(), UWorld::StaticClass(), FName( *(LOCTEXT("PreviewWorld", "Preview").ToString()) ));
 	PreviewWorld = TStrongObjectPtr<UWorld>(NewObject< UWorld >(GetTransientPackage(), UniqueWorldName));
@@ -192,8 +198,6 @@ FDataprepEditor::~FDataprepEditor()
 		PreviewWorld.Reset();
 	}
 
-	FCoreUObjectDelegates::OnObjectPropertyChanged.RemoveAll(this);
-
 	auto DeleteDirectory = [&](const FString& DirectoryToDelete)
 	{
 		const FString AbsolutePath = FPaths::ConvertRelativePathToFull( DirectoryToDelete );
@@ -401,10 +405,15 @@ void FDataprepEditor::InitDataprepEditor(const EToolkitMode::Type Mode, const TS
 	DataprepAssetInterfacePtr = TWeakObjectPtr<UDataprepAssetInterface>(InDataprepAssetInterface);
 	check( DataprepAssetInterfacePtr.IsValid() );
 
-	bIsDataprepInstance = Cast<UDataprepAssetInstance>(InDataprepAssetInterface) != nullptr;
+	bIsDataprepInstance = InDataprepAssetInterface->IsA<UDataprepAssetInstance>();
 
 	DataprepAssetInterfacePtr->GetOnChanged().AddSP( this, &FDataprepEditor::OnDataprepAssetChanged );
 
+	if ( UDataprepAsset* DataprepAsset = Cast<UDataprepAsset>( InDataprepAssetInterface ) )
+	{
+		DataprepAsset->GetOnStepObjectsAboutToBeRemoved().AddSP( this, &FDataprepEditor::OnStepObjectsAboutToBeDeleted );
+	}
+
 	// Assign unique session identifier
 	SessionID = FGuid::NewGuid().ToString();
 
@@ -568,6 +577,8 @@ void FDataprepEditor::OnBuildWorld()
 		return;
 	}
 
+	UpdateDataForPreviewSystem();
+
 	CachedAssets.Empty(Assets.Num());
 	for(TWeakObjectPtr<UObject>& WeakObject : Assets)
 	{
@@ -725,6 +736,8 @@ void FDataprepEditor::OnExecutePipeline()
 		RestoreFromSnapshot();
 	}
 
+	UpdateDataForPreviewSystem();
+
 	// Redraw 3D viewport
 	SceneViewportView->UpdateScene();
 
@@ -802,6 +815,8 @@ void FDataprepEditor::OnCommitWorld()
 	}
 
 	ResetBuildWorld();
+
+	UpdateDataForPreviewSystem();
 }
 
 void FDataprepEditor::ExtendMenu()
@@ -1244,6 +1259,84 @@ void FDataprepEditor::OnActionsContextChanged( const UDataprepActionAsset* Actio
 	}
 }
 
+void FDataprepEditor::UpdateDataForPreviewSystem()
+{
+	TArray<UObject*> ObjectsForThePreviewSystem;
+	FDataprepCoreUtils::GetActorsFromWorld( PreviewWorld.Get(), ObjectsForThePreviewSystem );
+	ObjectsForThePreviewSystem.Reserve( Assets.Num() );
+	for ( TWeakObjectPtr<UObject>& WeakObjectPtr : Assets )
+	{
+		if ( UObject* Object = WeakObjectPtr.Get() )
+		{
+			ObjectsForThePreviewSystem.Add( Object );
+		}
+	}
+	PreviewSystem->UpdateDataToProcess( MakeArrayView( ObjectsForThePreviewSystem ) );
+}
+
+bool FDataprepEditor::IsPreviewingStep(const UDataprepParameterizableObject* StepObject) const 
+{
+	return PreviewSystem->IsObservingObject( StepObject );
+}
+
+void FDataprepEditor::OnStepObjectsAboutToBeDeleted(const TArrayView<UDataprepParameterizableObject*>& StepObjects)
+{
+	// Todo Implement removal of objects from the preview system
+}
+
+void FDataprepEditor::PostUndo(bool bSuccess)
+{
+	if ( bSuccess )
+	{
+		if ( PreviewSystem->HasObservedObjects() )
+		{
+			// Check if a object as disappear
+			PreviewSystem->RestartProcessing();
+			// Set Preview view
+		}
+	}
+}
+
+void FDataprepEditor::PostRedo(bool bSuccess)
+{
+	if ( bSuccess )
+	{
+		if ( PreviewSystem->HasObservedObjects() )
+		{
+			// Check if a object as disappear
+			PreviewSystem->RestartProcessing();
+			// Set Preview view
+		}
+	}
+}
+
+void FDataprepEditor::SetPreviewedObjects(const TArrayView<UDataprepParameterizableObject*>& ObservedObjects)
+{
+	// (possible improvement) A validation that the observed object are from the same action in that the action is from the edited dataprep could help here.
+
+	if ( !PreviewSystem->HasObservedObjects() )
+	{ 
+		SceneOutliner->RemoveColumn( SceneOutliner::FBuiltInColumnTypes::ActorInfo() );
+		
+		SceneOutliner::FColumnInfo ColumnInfo( SceneOutliner::EColumnVisibility::Visible, 100, FCreateSceneOutlinerColumn::CreateLambda([Preview = PreviewSystem](ISceneOutliner& InSceneOutliner) -> TSharedRef< ISceneOutlinerColumn >
+			{
+				return MakeShared<FDataprepPreviewOutlinerColumn>( InSceneOutliner, Preview );
+			}));
+
+		SceneOutliner->AddColumn( FDataprepPreviewOutlinerColumn::ColumnID, ColumnInfo );
+
+		AssetPreviewView->AddColumn( MakeShared<FDataprepPreviewAssetColumn>( PreviewSystem ) );
+	}
+
+	PreviewSystem->SetObservedObjects( ObservedObjects );
+	
+	if ( SDataprepGraphEditor* RawGraphEditor = GraphEditor.Get() )
+	{
+		// Refresh the graph
+		RawGraphEditor->NotifyGraphChanged();
+	}
+}
+
 void DataprepEditorUtil::DeleteTemporaryPackage( const FString& PathToDelete )
 {
 	// See ContentBrowserUtils::LoadAssetsIfNeeded
diff --git a/Engine/Plugins/Enterprise/DataprepEditor/Source/DataprepEditor/Private/DataprepEditor.h b/Engine/Plugins/Enterprise/DataprepEditor/Source/DataprepEditor/Private/DataprepEditor.h
--- a/Engine/Plugins/Enterprise/DataprepEditor/Source/DataprepEditor/Private/DataprepEditor.h
+++ b/Engine/Plugins/Enterprise/DataprepEditor/Source/DataprepEditor/Private/DataprepEditor.h
@@ -2,13 +2,12 @@
 
 #pragma once
 
-#include "CoreMinimal.h"
-
-#include "DataprepGraph/DataprepGraph.h"
-
-#include "DataprepAssetInterface.h"
 #include "DataprepActionAsset.h"
+#include "DataprepAssetInterface.h"
+#include "DataprepGraph/DataprepGraph.h"
+#include "PreviewSystem/DataprepPreviewSystem.h"
 
+#include "CoreMinimal.h"
 #include "EditorUndoClient.h"
 #include "Engine/EngineBaseTypes.h"
 #include "GraphEditor.h"
@@ -19,11 +18,14 @@
 #include "UObject/SoftObjectPath.h"
 #include "UObject/StrongObjectPtr.h"
 
+class UDataprepParameterizableObject;
+
 class FSpawnTabArgs;
 class FTabManager;
 class FUICommandList;
 class IMessageLogListing;
 class IMessageToken;
+class SDataprepGraphEditor;
 class SDockableTab;
 class SGraphEditor;
 class SGraphNodeDetailsWidget;
@@ -110,6 +112,12 @@ public:
 	/** Set the selection of the world items */
 	void SetWorldObjectsSelection(TSet<TWeakObjectPtr<UObject>>&& NewSelection, EWorldSelectionFrom SelectionFrom = EWorldSelectionFrom::Unknow);
 
+	/** Setup the preview system to observe those steps */
+	void SetPreviewedObjects(const TArrayView<UDataprepParameterizableObject*>& ObservedObjects);
+
+	/** Is the preview system observing this step */
+	bool IsPreviewingStep(const UDataprepParameterizableObject* StepObject) const;
+
 private:
 	void BindCommands();
 	void OnSaveScene();
@@ -189,6 +197,14 @@ private:
 	/** Handles change to the content passed to an action */
 	void OnActionsContextChanged( const UDataprepActionAsset* ActionAsset, bool bWorldChanged, bool bAssetsChanged, const TArray< TWeakObjectPtr<UObject> >& NewAssets );
 
+	/** Update the data used by the preview system */
+	void UpdateDataForPreviewSystem();
+
+	void OnStepObjectsAboutToBeDeleted(const TArrayView<UDataprepParameterizableObject*>& StepObject);
+
+	virtual void PostUndo( bool bSuccess ) override;
+	virtual void PostRedo( bool bSuccess ) override;
+
 private:
 	bool bWorldBuilt;
 	bool bIsFirstRun;
@@ -206,7 +222,7 @@ private:
 	TSharedPtr<SGraphNodeDetailsWidget> DetailsView;
 	TSharedPtr<class SDataprepAssetView > DataprepAssetView;
 	TSharedPtr<SGraphEditor> PipelineView;
-	TSharedPtr<SGraphEditor> GraphEditor;
+	TSharedPtr<SDataprepGraphEditor> GraphEditor;
 
 	TSharedPtr<class ICustomSceneOutliner> SceneOutliner;
 
@@ -258,6 +274,9 @@ private:
 	/** Helper member to record classes of assets' sub-objects */
 	TMap<FString, UClass*> SnapshotClassesMap;
 
+	/** Steps Preview System */
+	TSharedRef<FDataprepPreviewSystem> PreviewSystem;
+
 	/**	The tab ids for all the tabs used */
 	static const FName ScenePreviewTabId;
 	static const FName AssetPreviewTabId;
diff --git a/Engine/Plugins/Enterprise/DataprepEditor/Source/DataprepEditor/Private/DataprepEditorStyle.cpp b/Engine/Plugins/Enterprise/DataprepEditor/Source/DataprepEditor/Private/DataprepEditorStyle.cpp
--- a/Engine/Plugins/Enterprise/DataprepEditor/Source/DataprepEditor/Private/DataprepEditorStyle.cpp
+++ b/Engine/Plugins/Enterprise/DataprepEditor/Source/DataprepEditor/Private/DataprepEditorStyle.cpp
@@ -12,6 +12,7 @@
 #include "Styling/SlateTypes.h"
 
 #define IMAGE_PLUGIN_BRUSH( RelativePath, ... ) FSlateImageBrush( FDataprepEditorStyle::InContent( RelativePath, ".png" ), __VA_ARGS__ )
+#define BOX_BRUSH( RelativePath, ... ) FSlateBoxBrush( FDataprepEditorStyle::InContent( RelativePath, ".png" ), __VA_ARGS__ )
 
 #define DEFAULT_FONT(...) FCoreStyle::GetDefaultFontStyle(__VA_ARGS__)
 
@@ -62,6 +63,8 @@ void FDataprepEditorStyle::Initialize()
 
 	StyleSet->Set("DataprepEditor.TrackNode.Slot", new IMAGE_PLUGIN_BRUSH("CircleBox", Icon32x32));
 
+	StyleSet->Set("DataprepEditor.Node.Body", new BOX_BRUSH("Node_Body", FMargin(16.f/64.f, 25.f/64.f, 16.f/64.f, 16.f/64.f)));
+
 	StyleSet->Set( "DataprepEditor.SectionFont", DEFAULT_FONT( "Bold", 10 ) );
 
 	// Dataprep action UI
@@ -122,6 +125,18 @@ void FDataprepEditorStyle::Initialize()
 			TilteTextBlockStyle.SetFont(FCoreStyle::GetDefaultFontStyle("Bold", 10));
 			StyleSet->Set("DataprepActionBlock.TitleTextBlockStyle", TilteTextBlockStyle);
 		}
+
+
+		StyleSet->Set("Graph.ActionStepNode.PreviewColor", FLinearColor(1.f, 0.87f, 0.0117f, 1.f));
+
+		{
+			FTextBlockStyle TilteTextBlockStyle = FEditorStyle::GetWidgetStyle< FTextBlockStyle >("NormalText");
+			TilteTextBlockStyle.SetFont(FCoreStyle::GetDefaultFontStyle("Italic", 7))
+				.SetShadowOffset(FVector2D::ZeroVector)
+				.SetColorAndOpacity(GetColor("Graph.ActionStepNode.PreviewColor"))
+				.SetShadowColorAndOpacity(FLinearColor(0.f, 0.f, 0.f, 0.7f));
+			StyleSet->Set("DataprepActionBlock.PreviewTextBlockStyle", TilteTextBlockStyle);
+		}
 	}
 
 	// DataprepGraphEditor
diff --git a/Engine/Plugins/Enterprise/DataprepEditor/Source/DataprepEditor/Private/DataprepGraph/DataprepGraph.cpp b/Engine/Plugins/Enterprise/DataprepEditor/Source/DataprepEditor/Private/DataprepGraph/DataprepGraph.cpp
--- a/Engine/Plugins/Enterprise/DataprepEditor/Source/DataprepEditor/Private/DataprepGraph/DataprepGraph.cpp
+++ b/Engine/Plugins/Enterprise/DataprepEditor/Source/DataprepEditor/Private/DataprepGraph/DataprepGraph.cpp
@@ -63,7 +63,8 @@ void FDataprepEditor::CreateGraphEditor()
 		GraphEditor = SNew(SDataprepGraphEditor, DataprepAsset)
 			.Appearance(AppearanceInfo)
 			.TitleBar(TitleBarWidget)
-			.GraphToEdit(DataprepGraph.Get());
+			.GraphToEdit(DataprepGraph.Get())
+			.DataprepEditor( StaticCastSharedRef<FDataprepEditor>( this->AsShared() ) );
 	}
 }
 
diff --git a/Engine/Plugins/Enterprise/DataprepEditor/Source/DataprepEditor/Private/DataprepGraph/DataprepGraphActionNode.h b/Engine/Plugins/Enterprise/DataprepEditor/Source/DataprepEditor/Private/DataprepGraph/DataprepGraphActionNode.h
--- a/Engine/Plugins/Enterprise/DataprepEditor/Source/DataprepEditor/Private/DataprepGraph/DataprepGraphActionNode.h
+++ b/Engine/Plugins/Enterprise/DataprepEditor/Source/DataprepEditor/Private/DataprepGraph/DataprepGraphActionNode.h
@@ -54,6 +54,9 @@ protected:
 
 	UPROPERTY()
 	int32 StepIndex;
+
+	// Is this node currently driving the filter preview
+	bool bIsPreviewed = false;
 };
 
 /**
diff --git a/Engine/Plugins/Enterprise/DataprepEditor/Source/DataprepEditor/Private/PreviewSystem/DataprepPreviewAssetColumn.cpp b/Engine/Plugins/Enterprise/DataprepEditor/Source/DataprepEditor/Private/PreviewSystem/DataprepPreviewAssetColumn.cpp
new file mode 100644
--- /dev/null
+++ b/Engine/Plugins/Enterprise/DataprepEditor/Source/DataprepEditor/Private/PreviewSystem/DataprepPreviewAssetColumn.cpp
@@ -0,0 +1,119 @@
+// Copyright Epic Games, Inc. All Rights Reserved.
+
+#include "PreviewSystem/DataprepPreviewAssetColumn.h"
+
+#include "PreviewSystem/DataprepPreviewSystem.h"
+#include "Widgets/SDataprepPreviewRow.h"
+
+#include "Widgets/Views/SHeaderRow.h"
+
+#define LOCTEXT_NAMESPACE "DataprepPreviewAssetColumn"
+
+const FName FDataprepPreviewAssetColumn::ColumnID = FName("PreviewSystem");
+
+FDataprepPreviewAssetColumn::FDataprepPreviewAssetColumn(const TSharedRef<FDataprepPreviewSystem>& InPreviewSystem)
+	: AssetPreviewWidget::IAssetPreviewColumn()
+	, PreviewSystem( InPreviewSystem )
+{
+}
+
+SHeaderRow::FColumn::FArguments FDataprepPreviewAssetColumn::ConstructHeaderRowColumn(const TSharedRef<AssetPreviewWidget::SAssetsPreviewWidget>& PreviewWidget)
+{
+	PreviewWidgetWeakPtr = PreviewWidget;
+	PreviewSystem->GetOnPreviewIsDoneProcessing().AddSP( this, &FDataprepPreviewAssetColumn::OnPreviewSytemIsDoneProcessing );
+
+	return SHeaderRow::Column( GetColumnID() )
+		.DefaultLabel( LOCTEXT("Preview_HeaderText", "Preview") )
+		.DefaultTooltip( LOCTEXT("Preview_HeaderTooltip", "Show the result of the current preview.") )
+		.FillWidth( 5.0f );
+}
+
+const TSharedRef<SWidget> FDataprepPreviewAssetColumn::ConstructRowWidget(const AssetPreviewWidget::IAssetTreeItemPtr& TreeItem, const STableRow<AssetPreviewWidget::IAssetTreeItemPtr>& Row, const TSharedRef<AssetPreviewWidget::SAssetsPreviewWidget>& PreviewWidget)
+{
+	if ( !TreeItem->IsFolder() )
+	{
+		if ( UObject* Object = static_cast<AssetPreviewWidget::FAssetTreeAssetItem&>( *TreeItem.Get() ).AssetPtr.Get() )
+		{
+			return SNew( SDataprepPreviewRow, PreviewSystem->GetPreviewDataForObject(Object) )
+				.HighlightText( PreviewWidget, &AssetPreviewWidget::SAssetsPreviewWidget::OnGetHighlightText );
+		}
+	}
+
+	return SNullWidget::NullWidget;
+}
+
+void FDataprepPreviewAssetColumn::PopulateSearchStrings(const AssetPreviewWidget::IAssetTreeItemPtr& Item, TArray< FString >& OutSearchStrings, const AssetPreviewWidget::SAssetsPreviewWidget& AssetPreview) const
+{
+	if ( !Item->IsFolder() )
+	{
+		if ( FDataprepPreviewProcessingResult* PreviewData = PreviewSystem->GetPreviewDataForObject( static_cast<AssetPreviewWidget::FAssetTreeAssetItem&>( *Item.Get() ).AssetPtr.Get() ).Get() )
+		{
+			PreviewData->PopulateSearchStringFromFetchedData( OutSearchStrings );
+		}
+	}
+}
+
+void FDataprepPreviewAssetColumn::SortItems(TArray<AssetPreviewWidget::IAssetTreeItemPtr>& OutItems, const EColumnSortMode::Type SortMode) const
+{
+	using IItemPtr = AssetPreviewWidget::IAssetTreeItemPtr;
+	OutItems.Sort([this, SortMode](const IItemPtr& First, const IItemPtr& Second)
+	{
+		if ( First->IsFolder() || Second->IsFolder() )
+		{
+			return First->Name < Second->Name;
+		}
+
+		FDataprepPreviewProcessingResult* FirstPreviewData = PreviewSystem->GetPreviewDataForObject( static_cast<AssetPreviewWidget::FAssetTreeAssetItem&>( *First.Get() ).AssetPtr.Get() ).Get();
+		FDataprepPreviewProcessingResult* SecondPreviewData = PreviewSystem->GetPreviewDataForObject( static_cast<AssetPreviewWidget::FAssetTreeAssetItem&>( *Second.Get() ).AssetPtr.Get() ).Get();
+
+		if ( FirstPreviewData && SecondPreviewData )
+		{
+			if ( FirstPreviewData->Status == SecondPreviewData->Status && FirstPreviewData->CurrentProcessingIndex == SecondPreviewData->CurrentProcessingIndex )
+			{
+				EDataprepPreviewResultComparaison Comparaison = FirstPreviewData->CompareFetchedDataTo( *SecondPreviewData );
+				if ( Comparaison != EDataprepPreviewResultComparaison::Equal )
+				{
+					if ( SortMode == EColumnSortMode::Descending )
+					{
+						return Comparaison == EDataprepPreviewResultComparaison::BiggerThan;
+					}
+					return Comparaison == EDataprepPreviewResultComparaison::SmallerThan;
+				}
+			}
+			else if ( FirstPreviewData->Status == EDataprepPreviewStatus::Pass )
+			{
+				return true;
+			}
+			else if ( SecondPreviewData->Status == EDataprepPreviewStatus::Pass)
+			{
+				return false;
+			}
+			else
+			{
+				// Filter by at which index they fail
+				if ( FirstPreviewData->CurrentProcessingIndex != SecondPreviewData->CurrentProcessingIndex )
+				{
+					return FirstPreviewData->CurrentProcessingIndex > SecondPreviewData->CurrentProcessingIndex;
+				}
+			}
+		}
+
+		// If all else fail filter by name (always Ascending)
+		return First->Name < Second->Name;
+	});
+	
+}
+
+
+void FDataprepPreviewAssetColumn::OnPreviewSytemIsDoneProcessing()
+{
+	if (TSharedPtr<AssetPreviewWidget::SAssetsPreviewWidget> PreviewWidget = PreviewWidgetWeakPtr.Pin())
+	{
+		if ( PreviewWidget->GetColumnSortMode( ColumnID ) != EColumnSortMode::None )
+		{
+			PreviewWidget->RequestSort();
+		}
+	}
+}
+
+#undef LOCTEXT_NAMESPACE
diff --git a/Engine/Plugins/Enterprise/DataprepEditor/Source/DataprepEditor/Private/PreviewSystem/DataprepPreviewAssetColumn.h b/Engine/Plugins/Enterprise/DataprepEditor/Source/DataprepEditor/Private/PreviewSystem/DataprepPreviewAssetColumn.h
new file mode 100644
--- /dev/null
+++ b/Engine/Plugins/Enterprise/DataprepEditor/Source/DataprepEditor/Private/PreviewSystem/DataprepPreviewAssetColumn.h
@@ -0,0 +1,45 @@
+// Copyright Epic Games, Inc. All Rights Reserved.
+
+#pragma once
+
+#include "Widgets/SAssetsPreviewWidget.h"
+
+#include "CoreMinimal.h"
+#include "Widgets/Views/SHeaderRow.h"
+
+class FDataprepPreviewSystem;
+
+class FDataprepPreviewAssetColumn : public AssetPreviewWidget::IAssetPreviewColumn
+{
+public:
+	FDataprepPreviewAssetColumn(const TSharedRef<FDataprepPreviewSystem>& InPreviewSystem);
+	virtual ~FDataprepPreviewAssetColumn() = default;
+
+
+	virtual uint8 GetCulumnPositionPriorityIndex() const override
+	{
+		return 60;
+	}
+
+	virtual FName GetColumnID() const override
+	{
+		return ColumnID;
+	}
+
+	virtual SHeaderRow::FColumn::FArguments ConstructHeaderRowColumn(const TSharedRef<AssetPreviewWidget::SAssetsPreviewWidget>& PreviewWidget) override;
+
+	virtual const TSharedRef<SWidget> ConstructRowWidget(const AssetPreviewWidget::IAssetTreeItemPtr& TreeItem, const STableRow<AssetPreviewWidget::IAssetTreeItemPtr>& Row, const TSharedRef<AssetPreviewWidget::SAssetsPreviewWidget>& PreviewWidget) override;
+
+	virtual void PopulateSearchStrings(const AssetPreviewWidget::IAssetTreeItemPtr& Item, TArray< FString >& OutSearchStrings, const AssetPreviewWidget::SAssetsPreviewWidget& AssetPreview) const override;
+
+	virtual void SortItems(TArray<AssetPreviewWidget::IAssetTreeItemPtr>& OutItems, const EColumnSortMode::Type SortMode) const override;
+
+	static const FName ColumnID;
+
+private:
+	// Request a sort refresh if this column was the one doing the sorting
+	void OnPreviewSytemIsDoneProcessing();
+
+	TWeakPtr<AssetPreviewWidget::SAssetsPreviewWidget> PreviewWidgetWeakPtr;
+	TSharedPtr<FDataprepPreviewSystem> PreviewSystem;
+};
diff --git a/Engine/Plugins/Enterprise/DataprepEditor/Source/DataprepEditor/Private/PreviewSystem/DataprepPreviewSceneOutlinerColumn.cpp b/Engine/Plugins/Enterprise/DataprepEditor/Source/DataprepEditor/Private/PreviewSystem/DataprepPreviewSceneOutlinerColumn.cpp
new file mode 100644
--- /dev/null
+++ b/Engine/Plugins/Enterprise/DataprepEditor/Source/DataprepEditor/Private/PreviewSystem/DataprepPreviewSceneOutlinerColumn.cpp
@@ -0,0 +1,158 @@
+// Copyright Epic Games, Inc. All Rights Reserved.
+
+#include "PreviewSystem/DataprepPreviewSceneOutlinerColumn.h"
+
+#include "PreviewSystem/DataprepPreviewSystem.h"
+#include "Widgets/SDataprepPreviewRow.h"
+
+#include "ActorTreeItem.h"
+#include "ComponentTreeItem.h"
+#include "ISceneOutliner.h"
+#include "SceneOutlinerVisitorTypes.h"
+#include "SceneOutlinerVisitorTypes.h"
+#include "SubComponentTreeItem.h"
+#include "UObject/Object.h"
+#include "UObject/WeakObjectPtrTemplates.h"
+#include "Widgets/SNullWidget.h"
+#include "Widgets/SWidget.h"
+
+#define LOCTEXT_NAMESPACE "DataprepPreviewOutlinerColumn"
+
+namespace DataprepPreviewOutlinerColumnUtils
+{
+	struct FObjectGetter : public SceneOutliner::TTreeItemGetter<TWeakObjectPtr<UObject>>
+	{
+		virtual TWeakObjectPtr<UObject> Get(const SceneOutliner::FActorTreeItem& ActorItem) const override { return ActorItem.Actor; };
+		virtual TWeakObjectPtr<UObject> Get(const SceneOutliner::FWorldTreeItem& WorldItem) const override { return nullptr; }
+		virtual TWeakObjectPtr<UObject> Get(const SceneOutliner::FFolderTreeItem& FolderItem) const override { return nullptr; }
+		virtual TWeakObjectPtr<UObject> Get(const SceneOutliner::FComponentTreeItem& ComponentItem) const override { return ComponentItem.Component; }
+		virtual TWeakObjectPtr<UObject> Get(const SceneOutliner::FSubComponentTreeItem& SubComponentItem) const override { return SubComponentItem.ParentComponent; }
+	};
+}
+
+const FName FDataprepPreviewOutlinerColumn::ColumnID = FName( TEXT("DataprepPreview") );
+
+FDataprepPreviewOutlinerColumn::FDataprepPreviewOutlinerColumn(ISceneOutliner& SceneOutliner, const TSharedRef<FDataprepPreviewSystem>& PreviewData)
+	: ISceneOutlinerColumn()
+	, WeakSceneOutliner( StaticCastSharedRef<ISceneOutliner>( SceneOutliner.AsShared() ) )
+	, CachedPreviewData( PreviewData )
+{
+}
+
+FName FDataprepPreviewOutlinerColumn::GetColumnID()
+{
+	return ColumnID;
+}
+
+SHeaderRow::FColumn::FArguments FDataprepPreviewOutlinerColumn::ConstructHeaderRowColumn()
+{
+	CachedPreviewData->GetOnPreviewIsDoneProcessing().AddSP( this, &FDataprepPreviewOutlinerColumn::OnPreviewSytemIsDoneProcessing );
+
+	return SHeaderRow::Column( GetColumnID() )
+		.DefaultLabel( LOCTEXT("Preview_HeaderText", "Preview") )
+		.DefaultTooltip( LOCTEXT("Preview_HeaderTooltip", "Show the result of the current preview.") )
+		.FillWidth( 5.0f );
+}
+
+const TSharedRef<SWidget> FDataprepPreviewOutlinerColumn::ConstructRowWidget(SceneOutliner::FTreeItemRef TreeItem, const STableRow<SceneOutliner::FTreeItemPtr>& Row)
+{
+	DataprepPreviewOutlinerColumnUtils::FObjectGetter Visitor; 
+	TreeItem->Visit( Visitor );
+
+	if ( UObject* Object = Visitor.Result().Get() )
+	{
+		if ( TSharedPtr<ISceneOutliner> SceneOutliner = WeakSceneOutliner.Pin() )
+		{ 
+			return SNew( SDataprepPreviewRow, CachedPreviewData->GetPreviewDataForObject( Object ) )
+				.HighlightText( SceneOutliner->GetFilterHighlightText() );
+		}
+	}
+	
+	return SNullWidget::NullWidget;
+}
+
+void FDataprepPreviewOutlinerColumn::PopulateSearchStrings(const SceneOutliner::ITreeItem& Item, TArray< FString >& OutSearchStrings) const
+{
+	DataprepPreviewOutlinerColumnUtils::FObjectGetter Visitor;
+	Item.Visit(Visitor);
+	if (UObject* Object = Visitor.Result().Get())
+	{
+		if ( TSharedPtr<FDataprepPreviewProcessingResult> PreviewResult = CachedPreviewData->GetPreviewDataForObject( Object ) )
+		{
+			PreviewResult->PopulateSearchStringFromFetchedData( OutSearchStrings );
+		}
+	}
+}
+
+void FDataprepPreviewOutlinerColumn::SortItems(TArray<SceneOutliner::FTreeItemPtr>& OutItems, const EColumnSortMode::Type SortMode) const
+{
+	OutItems.Sort([this, SortMode](const SceneOutliner::FTreeItemPtr& First, const SceneOutliner::FTreeItemPtr& Second)
+	{
+			DataprepPreviewOutlinerColumnUtils::FObjectGetter Visitor;
+			First->Visit( Visitor );
+			UObject* FirstObject = Visitor.Result().Get();
+			Second->Visit( Visitor );
+			UObject* SecondObject = Visitor.Result().Get();
+			if ( FirstObject && SecondObject )
+			{
+
+				FDataprepPreviewProcessingResult* FirstPreviewData = CachedPreviewData->GetPreviewDataForObject( FirstObject ).Get();
+				FDataprepPreviewProcessingResult* SecondPreviewData = CachedPreviewData->GetPreviewDataForObject( SecondObject ).Get();
+
+				if ( FirstPreviewData && SecondPreviewData )
+				{
+					if ( FirstPreviewData->Status == SecondPreviewData->Status && FirstPreviewData->CurrentProcessingIndex == SecondPreviewData->CurrentProcessingIndex )
+					{
+						EDataprepPreviewResultComparaison Comparaison = FirstPreviewData->CompareFetchedDataTo( *SecondPreviewData );
+						if ( Comparaison != EDataprepPreviewResultComparaison::Equal )
+						{
+							if ( SortMode == EColumnSortMode::Descending )
+							{
+								return Comparaison == EDataprepPreviewResultComparaison::BiggerThan;
+							}
+							return Comparaison == EDataprepPreviewResultComparaison::SmallerThan;
+						}
+					}
+					else if ( FirstPreviewData->Status == EDataprepPreviewStatus::Pass )
+					{
+						return true;
+					}
+					else if ( SecondPreviewData->Status == EDataprepPreviewStatus::Pass )
+					{
+						return false;
+					}
+					else
+					{
+						// Filter by at which index they fail
+						if ( FirstPreviewData->CurrentProcessingIndex != SecondPreviewData->CurrentProcessingIndex )
+						{
+							return FirstPreviewData->CurrentProcessingIndex > SecondPreviewData->CurrentProcessingIndex;
+						}
+					}
+				}
+			}
+
+		// If all else fail filter by name (always Ascending)
+		int32 SortPriorityFirst = First->GetTypeSortPriority();
+		int32 SortPrioritySecond = Second->GetTypeSortPriority();
+		if ( SortPriorityFirst != SortPriorityFirst )
+		{
+			return SortPriorityFirst < SortPrioritySecond;
+		}
+
+		return First->GetDisplayString() < Second->GetDisplayString();
+	});
+}
+
+void FDataprepPreviewOutlinerColumn::OnPreviewSytemIsDoneProcessing()
+{
+	if ( TSharedPtr<ISceneOutliner> SceneOutliner = WeakSceneOutliner.Pin() )
+	{
+		if ( SceneOutliner->GetColumnSortMode( ColumnID ) != EColumnSortMode::None )
+		{
+			SceneOutliner->RequestSort();
+		}
+	}
+}
+
+#undef LOCTEXT_NAMESPACE
diff --git a/Engine/Plugins/Enterprise/DataprepEditor/Source/DataprepEditor/Private/PreviewSystem/DataprepPreviewSceneOutlinerColumn.h b/Engine/Plugins/Enterprise/DataprepEditor/Source/DataprepEditor/Private/PreviewSystem/DataprepPreviewSceneOutlinerColumn.h
new file mode 100644
--- /dev/null
+++ b/Engine/Plugins/Enterprise/DataprepEditor/Source/DataprepEditor/Private/PreviewSystem/DataprepPreviewSceneOutlinerColumn.h
@@ -0,0 +1,40 @@
+// Copyright Epic Games, Inc. All Rights Reserved.
+
+#pragma once
+
+#include "CoreMinimal.h"
+#include "ISceneOutlinerColumn.h"
+
+class FDataprepPreviewSystem;
+class ISceneOutliner;
+class UObject;
+
+class FDataprepPreviewOutlinerColumn : public ISceneOutlinerColumn
+{
+public:
+	FDataprepPreviewOutlinerColumn(ISceneOutliner& SceneOutliner, const TSharedRef<FDataprepPreviewSystem>& PreviewData);
+	virtual ~FDataprepPreviewOutlinerColumn() = default;
+
+
+	virtual FName GetColumnID() override;
+
+	virtual SHeaderRow::FColumn::FArguments ConstructHeaderRowColumn() override;
+
+	virtual const TSharedRef<SWidget> ConstructRowWidget(SceneOutliner::FTreeItemRef TreeItem, const STableRow<SceneOutliner::FTreeItemPtr>& Row) override;
+
+	virtual void PopulateSearchStrings(const SceneOutliner::ITreeItem& Item, TArray< FString >& OutSearchStrings) const override;
+
+	virtual bool SupportsSorting() const override { return true; }
+
+	virtual void SortItems(TArray<SceneOutliner::FTreeItemPtr>& OutItems, const EColumnSortMode::Type SortMode) const override;
+
+	const static FName ColumnID;
+
+private:
+	// Request a sort refresh if this column was the one doing the sorting
+	void OnPreviewSytemIsDoneProcessing();
+
+	TWeakPtr<ISceneOutliner> WeakSceneOutliner;
+
+	TSharedPtr<FDataprepPreviewSystem> CachedPreviewData;
+};
diff --git a/Engine/Plugins/Enterprise/DataprepEditor/Source/DataprepEditor/Private/PreviewSystem/DataprepPreviewSystem.cpp b/Engine/Plugins/Enterprise/DataprepEditor/Source/DataprepEditor/Private/PreviewSystem/DataprepPreviewSystem.cpp
new file mode 100644
--- /dev/null
+++ b/Engine/Plugins/Enterprise/DataprepEditor/Source/DataprepEditor/Private/PreviewSystem/DataprepPreviewSystem.cpp
@@ -0,0 +1,442 @@
+// Copyright Epic Games, Inc. All Rights Reserved.
+
+#include "PreviewSystem/DataprepPreviewSystem.h"
+
+#include "DataprepBindingCommandChange.h"
+#include "DataprepParameterizableObject.h"
+#include "SelectionSystem/DataprepFilter.h"
+
+#include "Async/ParallelFor.h"
+#include "CoreGlobals.h"
+#include "Internationalization/Culture.h"
+#include "Internationalization/FastDecimalFormat.h"
+#include "Internationalization/Internationalization.h"
+#include "Internationalization/Internationalization.h"
+#include "Math/UnrealMathUtility.h"
+#include "Misc/ITransaction.h"
+#include "Stats/Stats.h"
+#include "UObject/UnrealType.h"
+
+namespace DataprepPreviewUtils
+{
+	struct FPreviewVariantComarator
+	{
+		FPreviewVariantComarator(FDataprepPreviewProcessingResult& InCallingObject)
+			: Caller( InCallingObject )
+		{}
+
+		EDataprepPreviewResultComparaison operator()(const FEmptyVariantState&)
+		{
+			return EDataprepPreviewResultComparaison::Equal;
+		}
+
+		EDataprepPreviewResultComparaison operator()(int32 Value)
+		{
+			int32 CallerValue = Caller.FetchedData.Get<int32>();
+			if ( CallerValue == Value )
+			{
+				return EDataprepPreviewResultComparaison::Equal;
+			}
+			return CallerValue > Value ? EDataprepPreviewResultComparaison::BiggerThan : EDataprepPreviewResultComparaison::SmallerThan;
+		}
+
+		EDataprepPreviewResultComparaison operator()(float Value)
+		{
+			float CallerValue = Caller.FetchedData.Get<float>();
+			if ( CallerValue == Value )
+			{
+				return EDataprepPreviewResultComparaison::Equal;
+			}
+			return CallerValue > Value ? EDataprepPreviewResultComparaison::BiggerThan : EDataprepPreviewResultComparaison::SmallerThan;
+		}
+
+		EDataprepPreviewResultComparaison operator()(const FString& Value)
+		{
+			const FString& CallerValue = Caller.FetchedData.Get<FString>();
+			int32 CompareResult = CallerValue.Compare( Value );
+			if ( CompareResult == 0 )
+			{
+				return EDataprepPreviewResultComparaison::Equal;
+			}
+			return CompareResult > 0 ? EDataprepPreviewResultComparaison::BiggerThan : EDataprepPreviewResultComparaison::SmallerThan;
+		}
+
+		template<class T>
+		EDataprepPreviewResultComparaison operator()(T) = delete;
+
+		FDataprepPreviewProcessingResult& Caller;
+	};
+
+	struct FGetTextFromVariant
+	{
+		FText operator()(const FEmptyVariantState&)
+		{
+			return FText::GetEmpty();
+		}
+
+		FText operator()(int32 Value)
+		{
+			return FText::AsNumber(Value);
+		}
+
+		FText operator()(float Value)
+		{
+			FNumberFormattingOptions Options = FInternationalization::Get().GetCurrentLocale()->GetDecimalNumberFormattingRules().CultureDefaultFormattingOptions;
+			Options.SetMaximumFractionalDigits(8);
+			Options.SetRoundingMode(ERoundingMode::FromZero);
+			return FText::AsNumber(Value, &Options);
+		}
+
+		FText operator()(const FString& Value)
+		{
+			return FText::FromString(Value);
+		}
+
+		// force explicit declaration of 
+		template <typename T>
+		FText operator()(T) = delete;
+	};
+
+	struct FAddSearchStringsFromVariant
+	{
+		FAddSearchStringsFromVariant(TArray<FString>& OutStrings)
+			: RefToOutStrings( OutStrings )
+		{
+		}
+
+		void operator()(const FEmptyVariantState&)
+		{
+		}
+
+		void operator()(int32 Value)
+		{
+			RefToOutStrings.Add( FText::AsNumber(Value).ToString() );
+		}
+
+		void operator()(float Value)
+		{
+			FNumberFormattingOptions Options = FInternationalization::Get().GetCurrentLocale()->GetDecimalNumberFormattingRules().CultureDefaultFormattingOptions;
+			Options.SetMaximumFractionalDigits(8);
+			Options.SetRoundingMode(ERoundingMode::FromZero);
+			RefToOutStrings.Add( FText::AsNumber(Value, &Options).ToString() );
+		}
+
+		void operator()(const FString& Value)
+		{
+			RefToOutStrings.Add( Value );
+		}
+
+		// force explicit declaration of 
+		template <typename T>
+		void operator()(T) = delete;
+
+	private:
+		TArray<FString>& RefToOutStrings;
+	};
+}
+
+const int32 FDataprepPreviewSystem::IncrementalCount = 2000;
+
+EDataprepPreviewResultComparaison FDataprepPreviewProcessingResult::CompareFetchedDataTo(const FDataprepPreviewProcessingResult& Other)
+{
+	if ( FetchedData.GetIndex() != Other.FetchedData.GetIndex() )
+	{
+		return EDataprepPreviewResultComparaison::Equal;
+	}
+
+	return Visit( DataprepPreviewUtils::FPreviewVariantComarator( *this ), Other.FetchedData );
+}
+
+FText FDataprepPreviewProcessingResult::GetFetchedDataAsText() const
+{
+	return Visit( DataprepPreviewUtils::FGetTextFromVariant(), FetchedData );
+}
+
+void FDataprepPreviewProcessingResult::PopulateSearchStringFromFetchedData(TArray<FString>& OutStrings) const
+{
+	Visit( DataprepPreviewUtils::FAddSearchStringsFromVariant( OutStrings ), FetchedData );
+}
+
+
+FDataprepPreviewSystem::~FDataprepPreviewSystem()
+{
+	StopTrackingObservedObjects();
+}
+
+void FDataprepPreviewSystem::UpdateDataToProcess(const TArrayView<UObject*>& Objects)
+{
+	PreviewResult.Empty( Objects.Num() );
+
+	for ( UObject* Object : Objects )
+	{
+		PreviewResult.Add( Object, MakeShared<FDataprepPreviewProcessingResult>() );
+	}
+
+	RestartProcessing();
+};
+
+void FDataprepPreviewSystem::SetObservedObjects(const TArrayView<UDataprepParameterizableObject*>& StepObjects)
+{
+	StopTrackingObservedObjects();
+
+	ObservedObjects.Empty( StepObjects.Num() );
+
+	// Reserve the double for the filters which we need to track their also their fetcher
+	ObservedOnPostEdit.Empty( StepObjects.Num() * 2 );
+
+	for ( UDataprepParameterizableObject* Object : StepObjects )
+	{
+		if ( Object )
+		{
+			ObservedObjects.Add( Object );
+			
+			FDelegateHandle Handle = Object->GetOnPostEdit().AddSP( this, &FDataprepPreviewSystem::OnObservedObjectPostEdit );
+			ObservedOnPostEdit.Add( Object, Handle );
+
+
+			if ( UDataprepFilter* Filter  = Cast<UDataprepFilter>( Object ) )
+			{
+				if ( UDataprepFetcher* Fetcher = Filter->GetFetcher() )
+				{
+					Handle = Fetcher->GetOnPostEdit().AddSP( this, &FDataprepPreviewSystem::OnObservedObjectPostEdit);
+					ObservedOnPostEdit.Add( Fetcher, Handle );
+				}
+			}
+		}
+	}
+
+	RestartProcessing();
+}
+
+void FDataprepPreviewSystem::AddReferencedObjects(FReferenceCollector& Collector)
+{
+	Collector.AddReferencedObjects( ObservedObjects );
+}
+
+TSharedPtr<FDataprepPreviewProcessingResult> FDataprepPreviewSystem::GetPreviewDataForObject(UObject* Object) const
+{
+	if ( const TSharedRef<FDataprepPreviewProcessingResult>* PreviewData = PreviewResult.Find( Object ) )
+	{
+		return *PreviewData;
+	}
+
+	return {};
+}
+
+void FDataprepPreviewSystem::IncrementalProcess()
+{
+	const int32 ItemsCount = PreviewResult.Num();
+	int32 BudgetLeft = IncrementalCount;
+	while ( BudgetLeft > 0 && CurrentProgress.CurrentFilterIndex < ObservedObjects.Num() )
+	{
+		int32 ItemToProcessCount = ItemsCount - CurrentProgress.CurrentObjectProcessed;
+
+		int32 ItemsProcessedCount = FMath::Min(ItemToProcessCount, BudgetLeft);
+
+		bool bDoneWithCurrentObject = false;
+
+		// We skip the current object if we are not able to process it
+		bool bDidSomeProcessing = false;
+
+		if ( CurrentProgress.Iterator.IsValid() )
+		{
+			if ( UDataprepFilter* Filter = Cast<UDataprepFilter>( ObservedObjects[CurrentProgress.CurrentFilterIndex] ) )
+			{
+				bDidSomeProcessing = true;
+				PrepareFilterBuffers( IncrementalCount );
+				ItemsProcessedCount = FillObjectsBuffer( ItemsProcessedCount );
+
+				TArrayView<UObject*> Objects( ObjectsBuffer.GetData(), ItemsProcessedCount );
+				TArrayView<FDataprepSelectionInfo> FilterResults( FilterResultsBuffer.GetData(), ItemsProcessedCount );
+				Filter->FilterAndGatherInfo( Objects, FilterResults );
+				PopulateResultFromFilter( ItemsProcessedCount );
+			}
+
+			if ( CurrentProgress.CurrentObjectProcessed >= ItemsCount )
+			{
+				bDoneWithCurrentObject = true;
+			}
+		}
+
+		bDoneWithCurrentObject |= !bDidSomeProcessing;
+
+		if ( bDidSomeProcessing )
+		{
+			BudgetLeft -= ItemsProcessedCount;
+		}
+
+		if ( bDoneWithCurrentObject )
+		{
+			CurrentProgress.CurrentFilterIndex++;
+			CurrentProgress.CurrentObjectProcessed = 0;
+			CurrentProgress.Iterator = MakeUnique<FResultIterator>( PreviewResult, false );
+		}
+	}
+}
+
+void FDataprepPreviewSystem::RestartProcessing()
+{
+	const bool bResetProcessStatus = CurrentProgress.CurrentFilterIndex != 0 || CurrentProgress.CurrentObjectProcessed != 0;
+
+	CurrentProgress.CurrentFilterIndex = 0;
+	CurrentProgress.CurrentObjectProcessed = 0;
+
+	if ( PreviewResult.Num() > 0 )
+	{
+		CurrentProgress.Iterator = MakeUnique<FResultIterator>( PreviewResult, false );
+
+		if ( bResetProcessStatus )
+		{
+			FResultIterator ResultIterator( PreviewResult, false );
+			while ( ResultIterator )
+			{
+				FDataprepPreviewProcessingResult& Result = ResultIterator.Value().Get();
+				Result.Status = EDataprepPreviewStatus::BeingProcessed;
+				Result.FetchedData.Set<FEmptyVariantState>( FEmptyVariantState() );
+				Result.CurrentProcessingIndex = 0;
+				++ResultIterator;
+			}
+		}
+
+		if ( ObservedObjects.Num() > 0 )
+		{
+			bIsProcessing = true;
+		}
+	}
+}
+
+void FDataprepPreviewSystem::ClearProcessing()
+{
+	SetObservedObjects( MakeArrayView<UDataprepParameterizableObject*>( nullptr, 0) );
+}
+
+bool FDataprepPreviewSystem::HasAnObjectObserved(const TArrayView<UDataprepParameterizableObject*>& StepObjects) const 
+{
+	for ( UDataprepParameterizableObject* StepObject : StepObjects )
+	{
+		if ( ObservedOnPostEdit.Contains( StepObject ) )
+		{
+			return true;
+		}
+	}
+
+	return false;
+}
+
+bool FDataprepPreviewSystem::IsObservingObject(const UDataprepParameterizableObject* StepObject) const
+{
+	return ObservedOnPostEdit.Contains( StepObject );
+}
+
+int32 FDataprepPreviewSystem::FillObjectsBuffer(int32 MaximunNumberOfObject)
+{
+	MaximunNumberOfObject = FMath::Min( ObjectsBuffer.Num(), MaximunNumberOfObject );
+
+	int32 Index = 0;
+	// Fill the object buffer
+	while ( Index < MaximunNumberOfObject && bool( *CurrentProgress.Iterator.Get() ) )
+	{
+		if ( CurrentProgress.Iterator->Value()->Status == EDataprepPreviewStatus::BeingProcessed )
+		{
+			ObjectsBuffer[Index] = CurrentProgress.Iterator->Key();
+			Index++;
+		}
+
+		CurrentProgress.CurrentObjectProcessed++;
+		CurrentProgress.Iterator->operator++();
+	}
+
+	return Index;
+}
+
+void FDataprepPreviewSystem::PrepareFilterBuffers(int32 DesiredSize)
+{
+	if ( ObjectsBuffer.Num() != DesiredSize )
+	{
+		ObjectsBuffer.Empty( DesiredSize );
+		ObjectsBuffer.AddZeroed( DesiredSize );
+	}
+
+	if ( FilterResultsBuffer.Num() != DesiredSize )
+	{
+		FilterResultsBuffer.Empty( DesiredSize );
+		FilterResultsBuffer.AddDefaulted( DesiredSize );
+	}
+}
+
+void FDataprepPreviewSystem::PopulateResultFromFilter(int32 NumberOfValidObjects)
+{
+	check( ObjectsBuffer.Num() == FilterResultsBuffer.Num() );
+
+	NumberOfValidObjects = FMath::Min( ObjectsBuffer.Num(), NumberOfValidObjects );
+	
+	const bool bIsLastFilter =  CurrentProgress.CurrentFilterIndex == ObservedObjects.Num() - 1;
+
+	ParallelFor( NumberOfValidObjects, [this, bIsLastFilter](int32 Index)
+		{
+			UObject* Object = ObjectsBuffer[Index];
+			if ( TSharedRef<FDataprepPreviewProcessingResult>* ResultPtr = PreviewResult.Find( Object ) )
+			{
+				FDataprepPreviewProcessingResult& Result = ResultPtr->Get();
+				Result.CurrentProcessingIndex = CurrentProgress.CurrentFilterIndex;
+				FDataprepSelectionInfo& FilterResult = FilterResultsBuffer[Index];
+				if ( FilterResult.bHasPassFilter )
+				{
+					if ( bIsLastFilter )
+					{
+						Result.Status = EDataprepPreviewStatus::Pass;
+						if ( FilterResult.bWasDataFetchedAndCached )
+						{
+							Result.FetchedData = MoveTemp( FilterResult.FetchedData );
+						}
+					}
+				}
+				else
+				{
+					Result.Status = EDataprepPreviewStatus::Failed;
+					if ( FilterResult.bWasDataFetchedAndCached && bIsLastFilter )
+					{
+						Result.FetchedData = MoveTemp( FilterResult.FetchedData );
+					}
+				}
+
+				FilterResult.bHasPassFilter = false;
+
+			}
+		});
+}
+
+void FDataprepPreviewSystem::StopTrackingObservedObjects()
+{
+	for ( auto& Pair : ObservedOnPostEdit )
+	{
+		if ( Pair.Key )
+		{
+			Pair.Key->GetOnPostEdit().Remove( Pair.Value );
+		}
+	}
+}
+
+void FDataprepPreviewSystem::OnObservedObjectPostEdit(UDataprepParameterizableObject& Object, FPropertyChangedChainEvent& PropertyChangedChainEvent)
+{
+	RestartProcessing();
+}
+
+
+void FDataprepPreviewSystem::Tick(float DeltaTime)
+{
+	IncrementalProcess();
+
+	if ( CurrentProgress.CurrentFilterIndex >= ObservedObjects.Num() )
+	{
+		bIsProcessing = false;
+		OnPreviewIsDoneProcessing.Broadcast();
+	}
+}
+
+TStatId FDataprepPreviewSystem::GetStatId() const
+{
+	RETURN_QUICK_DECLARE_CYCLE_STAT(FDataprepPreviewSystem, STATGROUP_Tickables);
+}
+
+
diff --git a/Engine/Plugins/Enterprise/DataprepEditor/Source/DataprepEditor/Private/PreviewSystem/DataprepPreviewSystem.h b/Engine/Plugins/Enterprise/DataprepEditor/Source/DataprepEditor/Private/PreviewSystem/DataprepPreviewSystem.h
new file mode 100644
--- /dev/null
+++ b/Engine/Plugins/Enterprise/DataprepEditor/Source/DataprepEditor/Private/PreviewSystem/DataprepPreviewSystem.h
@@ -0,0 +1,162 @@
+// Copyright Epic Games, Inc. All Rights Reserved.
+
+#pragma once
+
+#include "DataprepParameterizableObject.h"
+#include "SelectionSystem/DataprepSelectionSystemStructs.h"
+
+#include "CoreMinimal.h"
+#include "Delegates/IDelegateInstance.h"
+#include "Internationalization/Text.h"
+#include "Templates/SharedPointer.h"
+#include "Templates/UniquePtr.h"
+#include "TickableEditorObject.h"
+#include "UObject/GCObject.h"
+
+class FReferenceCollector;
+
+struct TStatId;
+
+enum class EDataprepPreviewStatus : uint8
+{
+	BeingProcessed,
+	NotSupported,
+	Failed,
+	Pass
+	// There is assumption on the order of these so that the expression Pass > Failed return true
+};
+
+enum class EDataprepPreviewResultComparaison : uint8
+{
+	SmallerThan,
+	Equal,
+	BiggerThan
+};
+
+struct FDataprepPreviewProcessingResult
+{
+	FDataprepPreviewProcessingResult()
+		: Status( EDataprepPreviewStatus::BeingProcessed )
+		, CurrentProcessingIndex( 0 )
+		, FetchedData()
+	{}
+
+	EDataprepPreviewResultComparaison CompareFetchedDataTo(const FDataprepPreviewProcessingResult& Other);
+
+	FText GetFetchedDataAsText() const;
+
+	void PopulateSearchStringFromFetchedData(TArray<FString>& OutStrings) const;
+
+	EDataprepPreviewStatus Status;
+	uint32 CurrentProcessingIndex;
+	FFilterVariantData FetchedData;
+};
+
+/**
+ * A class managing the responsibility of processing the data for the preview of the result of a filter
+ */
+class FDataprepPreviewSystem final : public FGCObject, public FTickableEditorObject, public TSharedFromThis<FDataprepPreviewSystem>
+{
+public:
+	FDataprepPreviewSystem() = default;
+	virtual ~FDataprepPreviewSystem();
+
+	FDataprepPreviewSystem(const FDataprepPreviewSystem&) = delete;
+	FDataprepPreviewSystem(FDataprepPreviewSystem&&) = delete;
+	FDataprepPreviewSystem& operator=(const FDataprepPreviewSystem&) = delete;
+	FDataprepPreviewSystem& operator=(FDataprepPreviewSystem&&) = delete;
+
+	// Update the data to process
+	void UpdateDataToProcess(const TArrayView<UObject*>& Objects);
+
+	void SetObservedObjects(const TArrayView<UDataprepParameterizableObject*>& Objects);
+
+	bool HasObservedObjects() const { return ObservedObjects.Num() > 0; }
+
+	// Begin FGCObject interface
+	virtual void AddReferencedObjects(FReferenceCollector& Collector) override;
+	// End FGCObject interface
+
+	TSharedPtr<FDataprepPreviewProcessingResult> GetPreviewDataForObject(UObject* Object) const;
+
+	// Reset the processing from the start
+	void RestartProcessing();
+
+	// Cancel any ongoing processing
+	void ClearProcessing();
+
+	bool HasAnObjectObserved(const TArrayView<UDataprepParameterizableObject*>& StepObjects) const;
+
+	// Is the preview system observing an object
+	bool IsObservingObject(const UDataprepParameterizableObject* StepObject) const;
+
+	// Broadcast when the preview system is done processing the data to preview
+	DECLARE_EVENT(FDataprepPreviewSystem, FOnPreviewIsDoneProcessing)
+	FOnPreviewIsDoneProcessing& GetOnPreviewIsDoneProcessing() { return OnPreviewIsDoneProcessing; }
+
+private:
+
+	void IncrementalProcess();
+
+	/**
+	 * Fill the object buffer from the current iterator and update the current progress
+	 * @param MaximunNumberOfObject Useful if we don't want to fill all the buffer
+	 * @return the number of object inserted into the buffer
+	 */
+	int32 FillObjectsBuffer(int32 MaximunNumberOfObject);
+
+	/**
+	 * Ensure that buffer are of the proper size
+	 */
+	void PrepareFilterBuffers(int32 DesiredSize);
+
+	/**
+	 * Migrate the result of the filter from the buffer to the preview result map
+	 * @param NumberOfValidObjects Allow to stop before the end of the buffer it wasn't fully use
+	 */
+	void PopulateResultFromFilter(int32 NumberOfValidObjects);
+
+	void StopTrackingObservedObjects();
+
+	void OnObservedObjectPostEdit(UDataprepParameterizableObject& Object, FPropertyChangedChainEvent& PropertyChangedChainEvent);
+
+public:
+
+	// Begin FTickableObjectBase interface
+	virtual void Tick(float DeltaTime) override;
+	virtual bool IsTickable() const override { return bIsProcessing; }
+	virtual TStatId GetStatId() const override;
+	// End FTickableObjectBase interface
+
+private:
+	// The filters previewed
+	TArray<UDataprepParameterizableObject*> ObservedObjects;
+
+	// The handle for the PostEditEvents of the observed objects
+	TMap<UDataprepParameterizableObject*,FDelegateHandle> ObservedOnPostEdit;
+
+	// The inprogress or done result
+	TMap<UObject*, TSharedRef<FDataprepPreviewProcessingResult>> PreviewResult;
+
+	using FResultIterator = TMap<UObject*, TSharedRef<FDataprepPreviewProcessingResult>>::TIterator;
+
+	// The current progress of processing for the preview
+	struct
+	{
+		int32 CurrentFilterIndex;
+		int32 CurrentObjectProcessed;
+		TUniquePtr<FResultIterator> Iterator;
+	} CurrentProgress;
+
+	bool bIsProcessing = false;
+
+	// Buffers of object to process
+	TArray<UObject*> ObjectsBuffer;
+
+	// Buffer for the filters
+	TArray<FDataprepSelectionInfo> FilterResultsBuffer;
+
+	static const int32 IncrementalCount;
+
+	FOnPreviewIsDoneProcessing OnPreviewIsDoneProcessing;
+};
diff --git a/Engine/Plugins/Enterprise/DataprepEditor/Source/DataprepEditor/Private/Widgets/DataprepGraph/SDataprepActionBlock.cpp b/Engine/Plugins/Enterprise/DataprepEditor/Source/DataprepEditor/Private/Widgets/DataprepGraph/SDataprepActionBlock.cpp
--- a/Engine/Plugins/Enterprise/DataprepEditor/Source/DataprepEditor/Private/Widgets/DataprepGraph/SDataprepActionBlock.cpp
+++ b/Engine/Plugins/Enterprise/DataprepEditor/Source/DataprepEditor/Private/Widgets/DataprepGraph/SDataprepActionBlock.cpp
@@ -216,35 +216,6 @@ FReply SDataprepActionBlock::OnMouseButtonDown(const FGeometry& MyGeometry, cons
 	}
 #endif
 
-	// Take ownership of the mouse if right mouse button clicked to display contextual menu
-	if ( MouseEvent.GetEffectingButton() == EKeys::RightMouseButton )
-	{
-		return FReply::Handled();
-	}
-
-	return FReply::Unhandled();
-}
-
-FReply SDataprepActionBlock::OnMouseButtonUp(const FGeometry& MyGeometry, const FPointerEvent& MouseEvent)
-{
-	if ( MouseEvent.GetEffectingButton() == EKeys::RightMouseButton )
-	{
-		FMenuBuilder MenuBuilder( true, nullptr );
-		PopulateMenuBuilder( MenuBuilder );
-
-		FWidgetPath WidgetPath = MouseEvent.GetEventPath() ? *MouseEvent.GetEventPath() : FWidgetPath();
-		FSlateApplication::Get().PushMenu(
-			AsShared(),
-			WidgetPath,
-			MenuBuilder.MakeWidget(),
-			MouseEvent.GetScreenSpacePosition(),
-			FPopupTransitionEffect(FPopupTransitionEffect::ContextMenu)
-		);
-
-		// Release mouse capture
-		return FReply::Handled().ReleaseMouseCapture();
-	}
-
 	return FReply::Unhandled();
 }
 
diff --git a/Engine/Plugins/Enterprise/DataprepEditor/Source/DataprepEditor/Private/Widgets/DataprepGraph/SDataprepActionBlock.h b/Engine/Plugins/Enterprise/DataprepEditor/Source/DataprepEditor/Private/Widgets/DataprepGraph/SDataprepActionBlock.h
--- a/Engine/Plugins/Enterprise/DataprepEditor/Source/DataprepEditor/Private/Widgets/DataprepGraph/SDataprepActionBlock.h
+++ b/Engine/Plugins/Enterprise/DataprepEditor/Source/DataprepEditor/Private/Widgets/DataprepGraph/SDataprepActionBlock.h
@@ -34,7 +34,6 @@ class SDataprepActionBlock : public SCompoundWidget
 		return FVector2D( FMath::Max(280.f, ChildSize.X) , ChildSize.Y );
 	}
 	virtual FReply OnMouseButtonDown(const FGeometry& MyGeometry, const FPointerEvent& MouseEvent) override;
-	virtual FReply OnMouseButtonUp(const FGeometry& MyGeometry, const FPointerEvent& MouseEvent) override;
 	virtual FReply OnDragDetected(const FGeometry& MyGeometry, const FPointerEvent& MouseEvent) override;
 	// End of SWidget Interface
 
diff --git a/Engine/Plugins/Enterprise/DataprepEditor/Source/DataprepEditor/Private/Widgets/DataprepGraph/SDataprepFilter.cpp b/Engine/Plugins/Enterprise/DataprepEditor/Source/DataprepEditor/Private/Widgets/DataprepGraph/SDataprepFilter.cpp
--- a/Engine/Plugins/Enterprise/DataprepEditor/Source/DataprepEditor/Private/Widgets/DataprepGraph/SDataprepFilter.cpp
+++ b/Engine/Plugins/Enterprise/DataprepEditor/Source/DataprepEditor/Private/Widgets/DataprepGraph/SDataprepFilter.cpp
@@ -2,6 +2,7 @@
 
 #include "Widgets/DataprepGraph/SDataprepFilter.h"
 
+#include "DataprepCoreUtils.h"
 #include "DataprepEditorStyle.h"
 #include "SelectionSystem/DataprepBoolFilter.h"
 #include "SelectionSystem/DataprepFilter.h"
@@ -15,6 +16,7 @@
 #include "Widgets/DataprepGraph/SDataprepStringFilter.h"
 #include "Widgets/DataprepWidgets.h"
 
+#include "EditorStyleSet.h"
 #include "Framework/Commands/UIAction.h"
 #include "Framework/MultiBox/MultiBoxBuilder.h"
 #include "ScopedTransaction.h"
@@ -30,6 +32,8 @@ void SDataprepFilter::Construct(const FArguments& InArgs, UDataprepFilter& InFil
 	TAttribute<FText> TooltipTextAttribute = MakeAttributeSP( this, &SDataprepFilter::GetTooltipText );
 	SetToolTipText( TooltipTextAttribute );
 
+	bIsPreviewed = InArgs._IsPreviewed;
+
 	SDataprepActionBlock::Construct( SDataprepActionBlock::FArguments().IsSimplified(InArgs._IsSimplified), InDataprepActionContext );
 }
 
@@ -46,8 +50,12 @@ void SDataprepFilter::Tick(const FGeometry& AllottedGeometry, const double InCur
 
 FSlateColor SDataprepFilter::GetOutlineColor() const
 {
-	static const FSlateColor OutlineColor = FDataprepEditorStyle::GetColor( "DataprepActionStep.Filter.OutlineColor" );
-	return OutlineColor;
+	if ( bIsPreviewed )
+	{
+		return FDataprepEditorStyle::GetColor( "Graph.ActionStepNode.PreviewColor" );
+	}
+
+	return FDataprepEditorStyle::GetColor( "DataprepActionStep.Filter.OutlineColor" );
 }
 
 FText SDataprepFilter::GetBlockTitle() const
@@ -75,13 +83,33 @@ TSharedRef<SWidget> SDataprepFilter::GetTitleWidget()
 	const ISlateStyle* DataprepEditorStyle = FSlateStyleRegistry::FindSlateStyle( FDataprepEditorStyle::GetStyleSetName() );
 	check( DataprepEditorStyle );
 	const float DefaultPadding = DataprepEditorStyle->GetFloat( "DataprepAction.Padding" );
-
-	return SNew( STextBlock )
+	TSharedRef<SWidget> DefaultTitle = SNew( STextBlock )
 		.Text( this, &SDataprepFilter::GetBlockTitle )
 		.TextStyle( &DataprepEditorStyle->GetWidgetStyle<FTextBlockStyle>( "DataprepActionBlock.TitleTextBlockStyle" ) )
 		.ColorAndOpacity( FLinearColor( 1.f, 1.f, 1.f ) )
 		.Margin( FMargin( DefaultPadding ) )
 		.Justification( ETextJustify::Center );
+
+
+	if ( bIsPreviewed )
+	{
+		return SNew( SVerticalBox )
+			+SVerticalBox::Slot()
+			.AutoHeight()
+			[
+				DefaultTitle
+			]
+			+ SVerticalBox::Slot()
+			.AutoHeight()
+			.Padding( FMargin( 2.0f, 4.0f, 0.0f, 0.0f ) )
+			[
+				SNew( STextBlock )
+				.TextStyle(  &DataprepEditorStyle->GetWidgetStyle<FTextBlockStyle>( "DataprepActionBlock.PreviewTextBlockStyle" ) )
+				.Text( LOCTEXT("PreviewLabel", "Previewing") )
+			];
+	}
+
+	return DefaultTitle;
 }
 
 TSharedRef<SWidget> SDataprepFilter::GetContentWidget()
@@ -150,15 +178,6 @@ void SDataprepFilter::InverseFilter()
 	{
 		FScopedTransaction Transaction( LOCTEXT("InverseFilterTransaction", "Inverse the filter") );
 		Filter->SetIsExcludingResult( !Filter->IsExcludingResult() );
-
-		FProperty* Property = Filter->GetClass()->FindPropertyByName( TEXT("bIsExcludingResult") );
-		check( Property );
-		FEditPropertyChain EditChain;
-		EditChain.AddHead( Property );
-		EditChain.SetActivePropertyNode( Property );
-		FPropertyChangedEvent EditPropertyChangeEvent( Property, EPropertyChangeType::ValueSet );
-		FPropertyChangedChainEvent EditChangeChainEvent( EditChain, EditPropertyChangeEvent );
-		Filter->PostEditChangeChainProperty( EditChangeChainEvent );
 	}
 }
 
diff --git a/Engine/Plugins/Enterprise/DataprepEditor/Source/DataprepEditor/Private/Widgets/DataprepGraph/SDataprepFilter.h b/Engine/Plugins/Enterprise/DataprepEditor/Source/DataprepEditor/Private/Widgets/DataprepGraph/SDataprepFilter.h
--- a/Engine/Plugins/Enterprise/DataprepEditor/Source/DataprepEditor/Private/Widgets/DataprepGraph/SDataprepFilter.h
+++ b/Engine/Plugins/Enterprise/DataprepEditor/Source/DataprepEditor/Private/Widgets/DataprepGraph/SDataprepFilter.h
@@ -13,17 +13,17 @@ class UDataprepFilter;
 
 struct FDataprepSchemaActionContext;
 
+
 class SDataprepFilter : public SDataprepActionBlock, public FGCObject
 {
 public:
-#ifndef NO_BLUEPRINT
 	SLATE_BEGIN_ARGS(SDataprepFilter)
-	: _IsSimplified(false)
+
+		: _IsSimplified(false)
+		, _IsPreviewed(false)
 	{}
 		SLATE_ARGUMENT( bool, IsSimplified )
-#else
-	SLATE_BEGIN_ARGS(SDataprepFilter) {}
-#endif
+		SLATE_ARGUMENT( bool, IsPreviewed )
 	SLATE_END_ARGS()
 
 	void Construct(const FArguments& InArgs, UDataprepFilter& InFilter, const TSharedRef<FDataprepSchemaActionContext>& InDataprepActionContext);
@@ -52,4 +52,6 @@ private:
 	TSharedPtr<class SDataprepDetailsView> DetailsView;
 
 	UDataprepFilter* Filter = nullptr;
+
+	bool bIsPreviewed;
 };
diff --git a/Engine/Plugins/Enterprise/DataprepEditor/Source/DataprepEditor/Private/Widgets/DataprepGraph/SDataprepGraphActionNode.cpp b/Engine/Plugins/Enterprise/DataprepEditor/Source/DataprepEditor/Private/Widgets/DataprepGraph/SDataprepGraphActionNode.cpp
--- a/Engine/Plugins/Enterprise/DataprepEditor/Source/DataprepEditor/Private/Widgets/DataprepGraph/SDataprepGraphActionNode.cpp
+++ b/Engine/Plugins/Enterprise/DataprepEditor/Source/DataprepEditor/Private/Widgets/DataprepGraph/SDataprepGraphActionNode.cpp
@@ -359,6 +359,7 @@ void SDataprepGraphActionNode::Construct(const FArguments& InArgs, UDataprepGrap
 	DataprepActionPtr->GetOnStepsOrderChanged().AddSP(this, &SDataprepGraphActionNode::OnStepsChanged);
 
 	GraphNode = InActionNode;
+	DataprepEditor = InArgs._DataprepEditor;
 
 	SetCursor(EMouseCursor::ResizeLeftRight);
 	UpdateGraphNode();
@@ -543,7 +544,7 @@ TSharedRef<SWidget> SDataprepGraphActionNode::CreateBackground(const TAttribute<
 		[
 			SNew(SImage)
 			.ColorAndOpacity(FLinearColor::White)
-			.Image(FEditorStyle::GetBrush( "Graph.StateNode.Body" ))
+			.Image(FDataprepEditorStyle::GetBrush( "DataprepEditor.Node.Body" ))
 		]
 
 		+SOverlay::Slot()
@@ -553,7 +554,7 @@ TSharedRef<SWidget> SDataprepGraphActionNode::CreateBackground(const TAttribute<
 		[
 			SNew(SImage)
 			.ColorAndOpacity(ColorAndOpacity)
-			.Image(FEditorStyle::GetBrush( "Graph.StateNode.Body" ))
+			.Image(FDataprepEditorStyle::GetBrush( "DataprepEditor.Node.Body" ))
 		];
 }
 
@@ -607,17 +608,21 @@ int32 SDataprepGraphActionNode::OnPaint(const FPaintArgs& Args, const FGeometry&
 
 FReply SDataprepGraphActionNode::OnMouseButtonDown(const FGeometry& MyGeometry, const FPointerEvent& MouseEvent)
 {
-	GetOwnerPanel()->SelectionManager.ClickedOnNode(GraphNode, MouseEvent);
 	BorderBackgroundColor.Set(FDataprepEditorStyle::GetColor("DataprepActionStep.DragAndDrop"));
 
 	if( MouseEvent.GetEffectingButton() ==  EKeys::LeftMouseButton )
 	{
+		GetOwnerPanel()->SelectionManager.ClickedOnNode( GraphNode, MouseEvent );
 		return FReply::Handled().DetectDrag( AsShared(), EKeys::LeftMouseButton );
 	}
 
 	// Take ownership of the mouse if right mouse button clicked to display contextual menu
 	if ( MouseEvent.GetEffectingButton() == EKeys::RightMouseButton )
 	{
+		if ( !GetOwnerPanel()->SelectionManager.SelectedNodes.Contains( GraphNode ) )
+		{
+			GetOwnerPanel()->SelectionManager.ClickedOnNode( GraphNode, MouseEvent );
+		}
 		return FReply::Handled();
 	}
 
@@ -786,7 +791,8 @@ void SDataprepGraphActionNode::PopulateActionStepListWidget()
 
 		ActionStepNode->Initialize(DataprepAction, Index);
 
-		TSharedPtr<SDataprepGraphActionStepNode> ActionStepGraphNode = SNew(SDataprepGraphActionStepNode, ActionStepNode, SharedThis(this));
+		TSharedPtr<SDataprepGraphActionStepNode> ActionStepGraphNode = SNew(SDataprepGraphActionStepNode, ActionStepNode, SharedThis(this))
+			.DataprepEditor(DataprepEditor);
 
 		if(TrackNodePtr.IsValid())
 		{
diff --git a/Engine/Plugins/Enterprise/DataprepEditor/Source/DataprepEditor/Private/Widgets/DataprepGraph/SDataprepGraphActionNode.h b/Engine/Plugins/Enterprise/DataprepEditor/Source/DataprepEditor/Private/Widgets/DataprepGraph/SDataprepGraphActionNode.h
--- a/Engine/Plugins/Enterprise/DataprepEditor/Source/DataprepEditor/Private/Widgets/DataprepGraph/SDataprepGraphActionNode.h
+++ b/Engine/Plugins/Enterprise/DataprepEditor/Source/DataprepEditor/Private/Widgets/DataprepGraph/SDataprepGraphActionNode.h
@@ -14,6 +14,7 @@ class SVerticalBox;
 class UDataprepActionAsset;
 class UDataprepGraphActionNode;
 class UDataprepGraphActionStepNode;
+class FDataprepEditor;
 
 /**
  * The SDataprepGraphActionNode class is the SGraphNode associated
@@ -23,6 +24,7 @@ class SDataprepGraphActionNode : public SGraphNode
 {
 public:
 	SLATE_BEGIN_ARGS(SDataprepGraphActionNode) {}
+		SLATE_ARGUMENT(TWeakPtr<FDataprepEditor>, DataprepEditor)
 	SLATE_END_ARGS()
 
 	void Construct(const FArguments& InArgs, UDataprepGraphActionNode* InActionNode);
@@ -103,5 +105,8 @@ private:
 	/** Array of strong pointers to the UEdGraphNodes created for the action's steps */
 	TArray<TStrongObjectPtr<UDataprepGraphActionStepNode>> EdGraphStepNodes;
 
+	/** A optional ptr to a dataprep editor */
+	TWeakPtr<FDataprepEditor> DataprepEditor;
+
 	friend SDataprepGraphActionProxyNode;
 };
diff --git a/Engine/Plugins/Enterprise/DataprepEditor/Source/DataprepEditor/Private/Widgets/DataprepGraph/SDataprepGraphActionStepNode.cpp b/Engine/Plugins/Enterprise/DataprepEditor/Source/DataprepEditor/Private/Widgets/DataprepGraph/SDataprepGraphActionStepNode.cpp
--- a/Engine/Plugins/Enterprise/DataprepEditor/Source/DataprepEditor/Private/Widgets/DataprepGraph/SDataprepGraphActionStepNode.cpp
+++ b/Engine/Plugins/Enterprise/DataprepEditor/Source/DataprepEditor/Private/Widgets/DataprepGraph/SDataprepGraphActionStepNode.cpp
@@ -5,6 +5,7 @@
 // Dataprep includes
 #include "DataprepActionAsset.h"
 #include "DataprepCoreUtils.h"
+#include "DataprepEditor.h"
 #include "DataprepEditorLogCategory.h"
 #include "DataprepEditorStyle.h"
 #include "DataprepGraph/DataprepGraphActionNode.h"
@@ -16,9 +17,9 @@
 #include "Widgets/DataprepGraph/SDataprepActionSteps.h"
 #include "Widgets/DataprepGraph/SDataprepFilter.h"
 #include "Widgets/DataprepGraph/SDataprepGraphActionNode.h"
-#include "Widgets/DataprepGraph/SDataprepSelectionTransform.h"
 #include "Widgets/DataprepGraph/SDataprepGraphTrackNode.h"
 #include "Widgets/DataprepGraph/SDataprepOperation.h"
+#include "Widgets/DataprepGraph/SDataprepSelectionTransform.h"
 
 // Engine Includes
 #include "Editor.h"
@@ -42,6 +43,7 @@ void SDataprepGraphActionStepNode::Construct(const FArguments& InArgs, UDataprep
 
 	ParentNodePtr = InParent;
 	GraphNode = InActionStepNode;
+	DataprepEditor = InArgs._DataprepEditor;
 
 	SetCursor(EMouseCursor::ResizeUpDown);
 	UpdateGraphNode();
@@ -114,7 +116,14 @@ void SDataprepGraphActionStepNode::UpdateGraphNode()
 		if ( UDataprepActionStep* ActionStep = StepData->DataprepActionStepPtr.Get())
 		{
 			UDataprepParameterizableObject* StepObject = ActionStep->GetStepObject();
-			UClass* StepType = FDataprepCoreUtils::GetTypeOfActionStep(StepObject);
+
+			bool bIsPreviewed = false;
+			if (TSharedPtr<FDataprepEditor> DataprepEditorPtr = DataprepEditor.Pin())
+			{
+				bIsPreviewed = DataprepEditorPtr->IsPreviewingStep( StepObject );
+			}
+
+			UClass* StepType = FDataprepCoreUtils::GetTypeOfActionStep( StepObject );
 			if (StepType == UDataprepOperation::StaticClass())
 			{
 				UDataprepOperation* Operation = static_cast<UDataprepOperation*>( StepObject );
@@ -123,7 +132,12 @@ void SDataprepGraphActionStepNode::UpdateGraphNode()
 			else if (StepType == UDataprepFilter::StaticClass())
 			{
 				UDataprepFilter* Filter = static_cast<UDataprepFilter*>( StepObject );
-				ActionStepBlockPtr = StaticCastSharedRef<SDataprepActionBlock>( SNew(SDataprepFilter, *Filter, StepData).IsSimplified(true) );
+
+				ActionStepBlockPtr = StaticCastSharedRef<SDataprepActionBlock>( 
+					SNew( SDataprepFilter, *Filter, StepData )
+						.IsPreviewed( bIsPreviewed )
+						.IsSimplified(true)
+					);
 			}
 			else if (StepType == UDataprepSelectionTransform::StaticClass())
 			{
@@ -228,16 +242,40 @@ FSlateColor SDataprepGraphActionStepNode::GetBorderBackgroundColor() const
 
 FReply SDataprepGraphActionStepNode::OnMouseButtonDown(const FGeometry& MyGeometry, const FPointerEvent& MouseEvent)
 {
-	GetOwnerPanel()->SelectionManager.ClickedOnNode(GraphNode, MouseEvent);
-
 	if ( MouseEvent.GetEffectingButton() ==  EKeys::LeftMouseButton )
 	{
+		GetOwnerPanel()->SelectionManager.ClickedOnNode( GraphNode, MouseEvent );
 		return FReply::Handled().DetectDrag( AsShared(), EKeys::LeftMouseButton );
 	}
 
+	// Take ownership of the mouse if right mouse button clicked to display contextual menu
+	if ( MouseEvent.GetEffectingButton() == EKeys::RightMouseButton )
+	{
+		if ( !GetOwnerPanel()->SelectionManager.SelectedNodes.Contains( GraphNode ) )
+		{
+			GetOwnerPanel()->SelectionManager.ClickedOnNode( GraphNode, MouseEvent );
+		}
+		return FReply::Handled();
+	}
+
 	return SGraphNode::OnMouseButtonDown( MyGeometry, MouseEvent);
 }
 
+FReply SDataprepGraphActionStepNode::OnMouseButtonUp(const FGeometry& MyGeometry, const FPointerEvent& MouseEvent)
+{
+	if ( MouseEvent.GetEffectingButton() == EKeys::RightMouseButton )
+	{
+		ensure(OwnerGraphPanelPtr.IsValid());
+
+		const FVector2D Position = MouseEvent.GetScreenSpacePosition();
+		OwnerGraphPanelPtr.Pin()->SummonContextMenu( Position, Position, GraphNode, nullptr, TArray<UEdGraphPin*>() );
+
+		return FReply::Handled();
+	}
+
+	return FReply::Unhandled();
+}
+
 FReply SDataprepGraphActionStepNode::OnMouseMove(const FGeometry& SenderGeometry, const FPointerEvent& MouseEvent)
 {
 	return FReply::Unhandled();
diff --git a/Engine/Plugins/Enterprise/DataprepEditor/Source/DataprepEditor/Private/Widgets/DataprepGraph/SDataprepGraphActionStepNode.h b/Engine/Plugins/Enterprise/DataprepEditor/Source/DataprepEditor/Private/Widgets/DataprepGraph/SDataprepGraphActionStepNode.h
--- a/Engine/Plugins/Enterprise/DataprepEditor/Source/DataprepEditor/Private/Widgets/DataprepGraph/SDataprepGraphActionStepNode.h
+++ b/Engine/Plugins/Enterprise/DataprepEditor/Source/DataprepEditor/Private/Widgets/DataprepGraph/SDataprepGraphActionStepNode.h
@@ -7,6 +7,7 @@
 #include "SGraphNode.h"
 #include "Widgets/DeclarativeSyntaxSupport.h"
 
+class FDataprepEditor;
 class SDataprepActionBlock;
 class SDataprepGraphActionNode;
 class SDataprepGraphTrackNode;
@@ -23,12 +24,14 @@ class SDataprepGraphActionStepNode : public SGraphNode
 {
 public:
 	SLATE_BEGIN_ARGS(SDataprepGraphActionStepNode) {}
+		SLATE_ARGUMENT(TWeakPtr<FDataprepEditor>, DataprepEditor)
 	SLATE_END_ARGS()
 
 	void Construct(const FArguments& InArgs, UDataprepGraphActionStepNode* InActionStepNode, const TSharedPtr<SDataprepGraphActionNode>& InParent);
 
 	// SWidget interface
 	virtual FReply OnMouseButtonDown( const FGeometry& MyGeometry, const FPointerEvent& MouseEvent ) override;
+	FReply OnMouseButtonUp( const FGeometry & MyGeometry, const FPointerEvent & MouseEvent ) override;
 	virtual FReply OnMouseMove(const FGeometry& SenderGeometry, const FPointerEvent& MouseEvent) override;
 	virtual FReply OnDragDetected(const FGeometry& MyGeometry, const FPointerEvent& MouseEvent);
 	// End of SWidget interface
@@ -77,4 +80,7 @@ private:
 
 	/** Pointer to the SDataprepGraphTrackNode displayed in the graph editor  */
 	TWeakPtr<SDataprepGraphActionNode> ParentNodePtr;
+
+	/** Optional pointer to the dataprep asset */
+	TWeakPtr<FDataprepEditor> DataprepEditor;
 };
diff --git a/Engine/Plugins/Enterprise/DataprepEditor/Source/DataprepEditor/Private/Widgets/DataprepGraph/SDataprepGraphEditor.cpp b/Engine/Plugins/Enterprise/DataprepEditor/Source/DataprepEditor/Private/Widgets/DataprepGraph/SDataprepGraphEditor.cpp
--- a/Engine/Plugins/Enterprise/DataprepEditor/Source/DataprepEditor/Private/Widgets/DataprepGraph/SDataprepGraphEditor.cpp
+++ b/Engine/Plugins/Enterprise/DataprepEditor/Source/DataprepEditor/Private/Widgets/DataprepGraph/SDataprepGraphEditor.cpp
@@ -3,14 +3,17 @@
 #include "Widgets/DataprepGraph/SDataprepGraphEditor.h"
 
 #include "DataprepAsset.h"
+#include "DataprepEditor.h"
 #include "DataprepEditorLogCategory.h"
 #include "DataprepEditorStyle.h"
 #include "DataprepGraph/DataprepGraph.h"
-#include "DataprepGraph/DataprepGraphSchema.h"
 #include "DataprepGraph/DataprepGraphActionNode.h"
+#include "DataprepGraph/DataprepGraphSchema.h"
+#include "DataprepParameterizableObject.h"
 #include "SchemaActions/DataprepAllMenuActionCollector.h"
 #include "SchemaActions/DataprepDragDropOp.h"
 #include "SchemaActions/IDataprepMenuActionCollector.h"
+#include "SelectionSystem/DataprepFilter.h"
 #include "Widgets/DataprepGraph/SDataprepGraphActionNode.h"
 #include "Widgets/DataprepGraph/SDataprepGraphActionStepNode.h"
 #include "Widgets/DataprepGraph/SDataprepGraphTrackNode.h"
@@ -43,9 +46,25 @@ const float SDataprepGraphEditor::TopPadding = 60.f;
 const float SDataprepGraphEditor::BottomPadding = 15.f;
 const float SDataprepGraphEditor::HorizontalPadding = 20.f;
 
-TSharedPtr<SDataprepGraphEditorNodeFactory> SDataprepGraphEditor::NodeFactory;
+TSharedPtr<FDataprepGraphEditorNodeFactory> SDataprepGraphEditor::NodeFactory;
 
-TSharedPtr<class SGraphNode> SDataprepGraphEditorNodeFactory::CreateNode(UEdGraphNode* Node) const
+TSharedPtr<SGraphNode> FDataprepGraphNodeFactory::CreateNodeWidget(UEdGraphNode* InNode)
+{
+	if (UDataprepGraphRecipeNode* RecipeNode = Cast<UDataprepGraphRecipeNode>(InNode))
+	{
+		return SNew(SDataprepGraphTrackNode, RecipeNode)
+			.NodeFactory( AsShared() );
+	}
+	else if (UDataprepGraphActionNode* ActionNode = Cast<UDataprepGraphActionNode>(InNode))
+	{
+		return SNew(SDataprepGraphActionNode, ActionNode)
+			.DataprepEditor( DataprepEditor );
+	}
+
+	return FGraphNodeFactory::CreateNodeWidget( InNode );
+}
+
+TSharedPtr<SGraphNode> FDataprepGraphEditorNodeFactory::CreateNode(UEdGraphNode* Node) const
 {
 	if (UDataprepGraphRecipeNode* RecipeNode = Cast<UDataprepGraphRecipeNode>(Node))
 	{
@@ -63,7 +82,7 @@ void SDataprepGraphEditor::RegisterFactories()
 {
 	if(!NodeFactory.IsValid())
 	{
-		NodeFactory  = MakeShareable( new SDataprepGraphEditorNodeFactory() );
+		NodeFactory  = MakeShared<FDataprepGraphEditorNodeFactory>();
 		FEdGraphUtilities::RegisterVisualNodeFactory(NodeFactory);
 	}
 }
@@ -81,6 +100,7 @@ void SDataprepGraphEditor::Construct(const FArguments& InArgs, UDataprepAsset* I
 {
 	check(InDataprepAsset);
 	DataprepAssetPtr = InDataprepAsset;
+	DataprepEditor = InArgs._DataprepEditor;
 
 	SGraphEditor::FGraphEditorEvents Events;
 	Events.OnCreateNodeOrPinMenu = SGraphEditor::FOnCreateNodeOrPinMenu::CreateSP(this, &SDataprepGraphEditor::OnCreateNodeOrPinMenu);
@@ -118,6 +138,8 @@ void SDataprepGraphEditor::Construct(const FArguments& InArgs, UDataprepAsset* I
 
 	FModifierKeysState ModifierKeyState = FSlateApplication::Get().GetModifierKeys();
 	bCachedControlKeyDown = ModifierKeyState.IsControlDown() || ModifierKeyState.IsCommandDown();
+
+	SetNodeFactory( MakeShared<FDataprepGraphNodeFactory>( InArgs._DataprepEditor ) );
 }
 
 // #ueent_toremove: Temp code for the nodes development
@@ -770,7 +792,10 @@ FActionMenuContent SDataprepGraphEditor::OnCreateNodeOrPinMenu(UEdGraph* Current
 	}
 
 	// Open contextual menu for action node
-	if(const UDataprepGraphActionNode* ActionNode = Cast<const UDataprepGraphActionNode>(InGraphNode))
+	UClass* GraphNodeClass = InGraphNode->GetClass();
+	const bool bIsActionNode = GraphNodeClass == UDataprepGraphActionNode::StaticClass();
+	const UDataprepGraphActionStepNode* FirstStepNode = Cast<UDataprepGraphActionStepNode>( InGraphNode );
+	if( bIsActionNode || FirstStepNode )
 	{
 		MenuBuilder->BeginSection( FName( TEXT("CommonSection") ), LOCTEXT("CommonSection", "Common") );
 		{
@@ -781,6 +806,86 @@ FActionMenuContent SDataprepGraphEditor::OnCreateNodeOrPinMenu(UEdGraph* Current
 		}
 		MenuBuilder->EndSection();
 
+		if ( FirstStepNode )
+		{
+			// Todo revisit this code the implementation is a bit ruff
+			// If all the all to node are filters (also we might need a more robust code path ex: register extension base on the base class/type of the steps?)
+			bool bIsSelectionOnlyFilters = true;
+			bool bAreFilterFromSameAction = true;
+			FGraphPanelSelectionSet SelectedNodes = GetSelectedNodes();
+			TArray<UDataprepFilter*> Filters;
+			Filters.Reserve( SelectedNodes.Num() );
+
+			const UDataprepActionAsset* ClikedAction = FirstStepNode->GetDataprepActionAsset();
+
+
+			for ( UObject* Node : SelectedNodes )
+			{
+				if ( UDataprepGraphActionStepNode* StepNode = Cast<UDataprepGraphActionStepNode>( Node ) )
+				{
+					if ( UDataprepActionAsset* Action = StepNode->GetDataprepActionAsset() )
+					{
+						bAreFilterFromSameAction &= Action == ClikedAction;
+
+						int32 StepIndex = StepNode->GetStepIndex();
+						if ( UDataprepActionStep* Step = Action->GetStep( StepIndex ).Get() )
+						{
+							if ( UDataprepParameterizableObject* StepObject = Step->GetStepObject() )
+							{
+								if ( UDataprepFilter* Filter = Cast<UDataprepFilter>( StepObject ) )
+								{
+									Filters.Add( Filter );
+									continue;
+								}
+							}
+						}
+					}
+				}
+
+				bIsSelectionOnlyFilters = false;
+				break;
+			}
+
+
+
+			if ( bIsSelectionOnlyFilters )
+			{
+				MenuBuilder->BeginSection( FName( TEXT("FilterSection") ), LOCTEXT("FilterSection", "Filter") );
+				{
+					FUIAction InverseFilterAction;
+					InverseFilterAction.ExecuteAction.BindLambda( [Filters]()
+						{
+							FScopedTransaction Transaction( LOCTEXT("InverseFilterTransaction", "Inverse the filter") );
+							for ( UDataprepFilter* Filter : Filters )
+							{
+								Filter->SetIsExcludingResult( !Filter->IsExcludingResult() );
+							}
+						});
+					MenuBuilder->AddMenuEntry(LOCTEXT("InverseFilter", "Inverse Filter(s) Selection"),
+						LOCTEXT("InverseFilterTooltip", "Inverse the resulting selection from a filter"),
+						FSlateIcon(),
+						InverseFilterAction);
+
+					if ( TSharedPtr<FDataprepEditor> DataprepEditorPtr = DataprepEditor.Pin() )
+					{ 
+						FUIAction SetFilterPreview;
+						SetFilterPreview.CanExecuteAction.BindLambda( [bAreFilterFromSameAction]() { return bAreFilterFromSameAction; } );
+						SetFilterPreview.ExecuteAction.BindLambda( [Filters, DataprepEditorPtr]()
+							{
+								TArray<UDataprepParameterizableObject*> ObjectsToObserve( Filters );
+								DataprepEditorPtr->SetPreviewedObjects( MakeArrayView<UDataprepParameterizableObject*>( ObjectsToObserve.GetData(), ObjectsToObserve.Num() ) );
+							});
+
+						MenuBuilder->AddMenuEntry(LOCTEXT("SetFilterPreview", "Preview Filter(s)"),
+							bAreFilterFromSameAction ? LOCTEXT("SetFilterPreviewTooltip", "Change the columns of the scene preview and asset preview tabs to display a preview of what the filters would select from the current scene") : LOCTEXT("SetFilterPreviewFailTooltip", "The filters can only be previewd if they are from the same action"),
+							FSlateIcon(),
+							SetFilterPreview);
+					}
+				}
+				MenuBuilder->EndSection();
+			}
+		}
+
 		return FActionMenuContent(MenuBuilder->MakeWidget());
 	}
 
diff --git a/Engine/Plugins/Enterprise/DataprepEditor/Source/DataprepEditor/Private/Widgets/DataprepGraph/SDataprepGraphEditor.h b/Engine/Plugins/Enterprise/DataprepEditor/Source/DataprepEditor/Private/Widgets/DataprepGraph/SDataprepGraphEditor.h
--- a/Engine/Plugins/Enterprise/DataprepEditor/Source/DataprepEditor/Private/Widgets/DataprepGraph/SDataprepGraphEditor.h
+++ b/Engine/Plugins/Enterprise/DataprepEditor/Source/DataprepEditor/Private/Widgets/DataprepGraph/SDataprepGraphEditor.h
@@ -13,11 +13,13 @@
 #include "Framework/Commands/UICommandList.h"
 #include "GraphEditor.h"
 #include "GraphEditorActions.h"
+#include "NodeFactory.h"
 #include "SGraphNode.h"
 #include "Widgets/DeclarativeSyntaxSupport.h"
 #include "Widgets/Layout/SConstraintCanvas.h"
 #include "Widgets/SWidget.h"
 
+class FDataprepEditor;
 class FUICommandList;
 class SDataprepGraphActionNode;
 class SDataprepGraphTrackNode;
@@ -27,8 +29,26 @@ class UDataprepAsset;
 class UBlueprint;
 class UEdGraph;
 
-class SDataprepGraphEditorNodeFactory : public FGraphPanelNodeFactory
+class FDataprepGraphNodeFactory : public FGraphNodeFactory, public TSharedFromThis<FDataprepGraphNodeFactory>
 {
+public:
+	virtual ~FDataprepGraphNodeFactory() = default;
+
+	template<class DataprepEditorPtrType>
+	FDataprepGraphNodeFactory(DataprepEditorPtrType&& InDataprepEditor)
+		: DataprepEditor(Forward<DataprepEditorPtrType>(InDataprepEditor))
+	{}
+
+	/** Create a widget for the supplied node */
+	virtual TSharedPtr<SGraphNode> CreateNodeWidget(UEdGraphNode* InNode) override;
+
+private:
+	TWeakPtr<FDataprepEditor> DataprepEditor;
+};
+
+class FDataprepGraphEditorNodeFactory : public FGraphPanelNodeFactory
+{
+public:
 	virtual TSharedPtr<SGraphNode> CreateNode(UEdGraphNode* Node) const override;
 };
 
@@ -49,6 +69,7 @@ public:
 		SLATE_ATTRIBUTE( FGraphAppearanceInfo, Appearance )
 		SLATE_ARGUMENT( UEdGraph*, GraphToEdit )
 		SLATE_ARGUMENT( FGraphEditorEvents, GraphEvents)
+		SLATE_ARGUMENT( TSharedPtr<FDataprepEditor>, DataprepEditor )
 	SLATE_END_ARGS()
 
 	void Construct(const FArguments& InArgs, UDataprepAsset* InDataprepAsset);
@@ -155,8 +176,10 @@ private:
 	static const float HorizontalPadding;
 
 	/** Factory to create the associated SGraphNode classes for Dataprep graph's UEdGraph classes */
-	static TSharedPtr<SDataprepGraphEditorNodeFactory> NodeFactory;
+	static TSharedPtr<FDataprepGraphEditorNodeFactory> NodeFactory;
 
+	// The dataprep editor that owns the graph
+	TWeakPtr<FDataprepEditor> DataprepEditor;
 
 	friend SDataprepGraphTrackNode;
 };
diff --git a/Engine/Plugins/Enterprise/DataprepEditor/Source/DataprepEditor/Private/Widgets/DataprepGraph/SDataprepGraphTrackNode.cpp b/Engine/Plugins/Enterprise/DataprepEditor/Source/DataprepEditor/Private/Widgets/DataprepGraph/SDataprepGraphTrackNode.cpp
--- a/Engine/Plugins/Enterprise/DataprepEditor/Source/DataprepEditor/Private/Widgets/DataprepGraph/SDataprepGraphTrackNode.cpp
+++ b/Engine/Plugins/Enterprise/DataprepEditor/Source/DataprepEditor/Private/Widgets/DataprepGraph/SDataprepGraphTrackNode.cpp
@@ -152,6 +152,8 @@ float SDataprepGraphTrackWidget::TrackDesiredHeight = 40.f;
 
 void SDataprepGraphTrackNode::Construct(const FArguments& InArgs, UDataprepGraphRecipeNode* InNode)
 {
+	NodeFactory = InArgs._NodeFactory;
+
 	bNodeDragging = false;
 	SetCursor(EMouseCursor::Default);
 	GraphNode = InNode;
@@ -208,7 +210,17 @@ void SDataprepGraphTrackNode::UpdateGraphNode()
 				NewActionNode->Initialize(ActionAsset, Index);
 
 				// #ueent_wip: Widget is created twice :-(
-				TSharedPtr< SDataprepGraphActionNode > ActionWidgetPtr = StaticCastSharedPtr<SDataprepGraphActionNode>(FNodeFactory::CreateNodeWidget(NewActionNode));
+				TSharedPtr<SGraphNode> ActionGraphNode;
+				if ( FGraphNodeFactory* GraphNodeFactor = NodeFactory.Get() )
+				{
+					ActionGraphNode = GraphNodeFactor->CreateNodeWidget( NewActionNode );
+				}
+				else
+				{
+					ActionGraphNode = FNodeFactory::CreateNodeWidget( NewActionNode );
+				}
+
+				TSharedPtr< SDataprepGraphActionNode > ActionWidgetPtr = StaticCastSharedPtr<SDataprepGraphActionNode>( ActionGraphNode );
 				if(SDataprepGraphActionNode* ActionWidget = ActionWidgetPtr.Get())
 				{
 					if(GraphPanelPtr.IsValid())
diff --git a/Engine/Plugins/Enterprise/DataprepEditor/Source/DataprepEditor/Private/Widgets/DataprepGraph/SDataprepGraphTrackNode.h b/Engine/Plugins/Enterprise/DataprepEditor/Source/DataprepEditor/Private/Widgets/DataprepGraph/SDataprepGraphTrackNode.h
--- a/Engine/Plugins/Enterprise/DataprepEditor/Source/DataprepEditor/Private/Widgets/DataprepGraph/SDataprepGraphTrackNode.h
+++ b/Engine/Plugins/Enterprise/DataprepEditor/Source/DataprepEditor/Private/Widgets/DataprepGraph/SDataprepGraphTrackNode.h
@@ -19,6 +19,7 @@
 #include "Widgets/Layout/SConstraintCanvas.h"
 #include "Widgets/SWidget.h"
 
+class FGraphNodeFactory;
 class SDataprepGraphActionNode;
 class SDataprepGraphEditor;
 class SDataprepGraphTrackNode;
@@ -60,6 +61,7 @@ class SDataprepGraphTrackNode : public SGraphNode
 {
 public:
 	SLATE_BEGIN_ARGS(SDataprepGraphTrackNode){}
+		SLATE_ARGUMENT(TSharedPtr<FGraphNodeFactory>, NodeFactory)
 	SLATE_END_ARGS()
 
 	void Construct(const FArguments& InArgs, UDataprepGraphRecipeNode* InNode);
@@ -150,5 +152,7 @@ private:
 	/** Array of strong pointers to the UEdGraphNodes created for the Dataprep asset's actions */
 	TArray<TStrongObjectPtr<UDataprepGraphActionNode>> EdGraphActionNodes;
 
+	TSharedPtr<FGraphNodeFactory> NodeFactory;
+
 	friend SDataprepGraphTrackWidget;
 };
diff --git a/Engine/Plugins/Enterprise/DataprepEditor/Source/DataprepEditor/Private/Widgets/SAssetsPreviewWidget.cpp b/Engine/Plugins/Enterprise/DataprepEditor/Source/DataprepEditor/Private/Widgets/SAssetsPreviewWidget.cpp
--- a/Engine/Plugins/Enterprise/DataprepEditor/Source/DataprepEditor/Private/Widgets/SAssetsPreviewWidget.cpp
+++ b/Engine/Plugins/Enterprise/DataprepEditor/Source/DataprepEditor/Private/Widgets/SAssetsPreviewWidget.cpp
@@ -10,189 +10,207 @@
 #include "SAssetSearchBox.h"
 #include "Misc/PackageName.h"
 
+#include "Math/UnrealMathUtility.h"
 #include "UObject/UObjectBaseUtility.h"
 
 #define LOCTEXT_NAMESPACE "AssetPreviewWidget"
 
 namespace AssetPreviewWidget
 {
-	struct FAssetTreeItem
+	/** This is the default column of the asset preview it contains the label and icon of the assets and folder */
+	class FAssetPreviewDefaultColumn : public IAssetPreviewColumn
 	{
-		// This is used to accelerate the construction of the tree in the set assets function
-		TMap< FString, FAssetTreeItemPtr> NameToFolder;
+	public:
+		virtual uint8 GetCulumnPositionPriorityIndex() const override
+		{
+			return 128;
+		}
 
-		// Childrens
-		TArray< FAssetTreeItemPtr > Folders;
-		TArray< FAssetTreeItemPtr > Assets;
+		virtual FName GetColumnID() const override
+		{
+			return ColumnID;
+		}
 
-		FString Name;
-		TWeakObjectPtr< UObject > AssetPtr;
+		virtual SHeaderRow::FColumn::FArguments ConstructHeaderRowColumn(const TSharedRef<SAssetsPreviewWidget>& PreviewWidget) override
+		{
+			return SHeaderRow::Column( ColumnID )
+				.DefaultLabel( LOCTEXT("AssetLabel_HeaderText", "Asset") )
+				.FillWidth( 5.0f );
+		}
 
-		TWeakPtr< SAssetsPreviewWidget > OwnerWeakPtr;
+		virtual const TSharedRef< SWidget > ConstructRowWidget(const IAssetTreeItemPtr& TreeItem, const STableRow<IAssetTreeItemPtr>& Row, const TSharedRef<SAssetsPreviewWidget>& PreviewWidget) override
+		{
+			PreviewWidgetWeakPtr = PreviewWidget;
 
-		// This value is cache for the last result for the filter function
-		bool bPassedFilter;
+			FSlateColor IconColor(FLinearColor::White);
 
-		void AddFolder(FAssetTreeItemPtr Folder)
-		{
-			if (Folder)
+			if ( !TreeItem->IsFolder() )
 			{
-				NameToFolder.Add(Folder->Name, Folder);
-				Folders.Add(Folder);
+				FAssetTreeAssetItem& AssetItem = static_cast<FAssetTreeAssetItem&>( *TreeItem.Get() );
+				UObject* Asset = AssetItem.AssetPtr.Get();
+				if ( Asset )
+				{
+					static FAssetToolsModule& AssetToolsModule = FModuleManager::LoadModuleChecked<FAssetToolsModule>( TEXT("AssetTools") );
+					TSharedPtr<IAssetTypeActions> AssetTypeActions = AssetToolsModule.Get().GetAssetTypeActionsForClass( Asset->GetClass() ).Pin();
+
+					if ( AssetTypeActions.IsValid() )
+					{
+						IconColor = FSlateColor( AssetTypeActions->GetTypeColor() );
+					}
+				}
 			}
+
+			return SNew(SHorizontalBox)
+				+ SHorizontalBox::Slot()
+				.AutoWidth()
+				.Padding(0, 0, 2, 0)
+				.VAlign(VAlign_Center)
+				[
+					// Item icon
+					SNew(SImage)
+					.Image( this, &FAssetPreviewDefaultColumn::GetIconBrush, TreeItem )
+					.ColorAndOpacity(IconColor)
+				]
+
+				+ SHorizontalBox::Slot()
+				.AutoWidth()
+				.VAlign(VAlign_Center)
+				[
+					SNew( STextBlock )
+					.Text( FText::FromString( *( TreeItem->Name ) ) )
+					.Font( FEditorStyle::GetFontStyle("ContentBrowser.SourceTreeItemFont") )
+					.HighlightText( PreviewWidget, &SAssetsPreviewWidget::OnGetHighlightText )
+				];
 		}
 
-		bool IsFolder() const
+		const FSlateBrush* GetIconBrush(IAssetTreeItemPtr TreeItem) const
 		{
-			if (Folders.Num() > 0 || Assets.Num() > 0)
+			const FSlateBrush* IconBrush = AssetIconBrush;
+			if ( TreeItem->IsFolder() )
 			{
-				return true;
+				const bool bExpanded = PreviewWidgetWeakPtr.Pin()->GetTreeView()->IsItemExpanded( TreeItem );
+				return bExpanded ? FolderOpenBrush : FolderClosedBrush;
 			}
-			return false;
+			return IconBrush;
 		}
 
-		bool Filter(const FText& FilterText)
+		virtual void PopulateSearchStrings(const IAssetTreeItemPtr& Item, TArray<FString>& OutSearchStrings, const SAssetsPreviewWidget& AssetPreview) const
 		{
-			bPassedFilter = false;
-			if (FilterText.IsEmpty())
+			if ( !Item->IsFolder() )
 			{
-				bPassedFilter = true;
+				OutSearchStrings.Append( AssetPreview.GetItemsName( static_cast<FAssetTreeAssetItem*>( Item.Get() )->AssetPtr ) );
 			}
+		}
 
-			if (IsFolder())
+		virtual void SortItems(TArray<IAssetTreeItemPtr>& OutItems, const EColumnSortMode::Type SortMode) const
+		{
+			if ( SortMode == EColumnSortMode::Ascending )
 			{
-				// A folder pass the filter if one of his child pass the filter
-				for (FAssetTreeItemPtr& Folder : Folders)
-				{
-					bPassedFilter = Folder->Filter(FilterText) || bPassedFilter;
-				}
-				for (FAssetTreeItemPtr& Asset : Assets)
-				{
-					bPassedFilter = Asset->Filter(FilterText) || bPassedFilter;
-				}
+				OutItems.Sort( [](const IAssetTreeItemPtr& First, const IAssetTreeItemPtr& Second) { return First->Name < Second->Name; } );
 			}
-			else if (!bPassedFilter && OwnerWeakPtr.IsValid())
+			
+			if ( SortMode == EColumnSortMode::Descending )
 			{
-				TArray< FString > FilterStrings;
-				FilterText.ToString().ParseIntoArray(FilterStrings, TEXT(" "));
+				OutItems.Sort( [](const IAssetTreeItemPtr& First, const IAssetTreeItemPtr& Second) { return First->Name > Second->Name; } );
+			}
+		}
 
-				TSharedPtr< SAssetsPreviewWidget > OwnerWidget = OwnerWeakPtr.Pin();
-				TArray< FString > ItemsName = OwnerWidget->GetItemsName(AssetPtr);
+		static const FName ColumnID;
 
-				// All the keywords must be pass for at least one of the items name in the hierarchy
-				bool bPassKeyWord = false;
-				for (const FString& KeyWord : FilterStrings)
-				{
-					bPassKeyWord = false;
-					for (const FString& ItemName : ItemsName)
-					{
-						if (ItemName.Contains(KeyWord))
-						{
-							bPassKeyWord = true;
-							break;
-						}
-					}
-					if (!bPassKeyWord)
-					{
-						break;
-					}
-				}
+	private:
+		/** Brushes for the different folder states */
+		static const FSlateBrush* FolderOpenBrush;
+		static const FSlateBrush* FolderClosedBrush;
+		static const FSlateBrush* AssetIconBrush;
 
-				bPassedFilter = bPassKeyWord;
+		TWeakPtr<SAssetsPreviewWidget> PreviewWidgetWeakPtr;
+	};
 
-			}
+	const FName FAssetPreviewDefaultColumn::ColumnID =  FName("DefaultColumn");
+	const FSlateBrush* FAssetPreviewDefaultColumn::FolderOpenBrush = FEditorStyle::GetBrush("ContentBrowser.AssetTreeFolderOpen");
+	const FSlateBrush* FAssetPreviewDefaultColumn::FolderClosedBrush = FEditorStyle::GetBrush("ContentBrowser.AssetTreeFolderClosed");
+	const FSlateBrush* FAssetPreviewDefaultColumn::AssetIconBrush = FEditorStyle::GetBrush("ContentBrowser.ColumnViewAssetIcon");
 
-			return bPassedFilter;
-		}
-	};
 
 	/** Represents a row in the AssetPreview's tree view */
-	class SAssetPreviewTableRow : public STableRow<FAssetTreeItemPtr>
+	class SAssetPreviewTableRow : public SMultiColumnTableRow<IAssetTreeItemPtr>
 	{
 	public:
 		SLATE_BEGIN_ARGS(SAssetPreviewTableRow) {}
 		SLATE_END_ARGS()
 
-		void Construct(const FArguments& InArgs, const TSharedRef<STableViewBase>& OwnerTableView, FAssetTreeItemPtr InItem, TSharedRef< const SAssetsPreviewWidget > InPreviewWidget)
+		void Construct(const FArguments& InArgs, const TSharedRef<STableViewBase>& OwnerTableView, IAssetTreeItemPtr InItem, TSharedRef< SAssetsPreviewWidget > InPreviewWidget)
 		{
-			FolderOpenBrush = FEditorStyle::GetBrush("ContentBrowser.AssetTreeFolderOpen");
-			FolderClosedBrush = FEditorStyle::GetBrush("ContentBrowser.AssetTreeFolderClosed");
-			AssetIconBrush = FEditorStyle::GetBrush("ContentBrowser.ColumnViewAssetIcon");
-
 			PreviewWidgetWeakPtr = InPreviewWidget;
 			ItemWeakPtr = InItem;
 
-			FSlateColor IconColor(FLinearColor::White);
-			TWeakObjectPtr< UObject > AssetPtr = InItem->AssetPtr;
-			if (AssetPtr.IsValid())
-			{
-				static FAssetToolsModule& AssetToolsModule = FModuleManager::LoadModuleChecked<FAssetToolsModule>(TEXT("AssetTools"));
-				TSharedPtr<IAssetTypeActions> AssetTypeActions = AssetToolsModule.Get().GetAssetTypeActionsForClass(AssetPtr->GetClass()).Pin();
-
-				if (AssetTypeActions.IsValid())
-				{
-					IconColor = FSlateColor(AssetTypeActions->GetTypeColor());
-				}
-			}
-
-			STableRow::Construct(
+			SMultiColumnTableRow<IAssetTreeItemPtr>::Construct(
 				STableRow::FArguments()
 				.Style(FEditorStyle::Get(), "ContentBrowser.AssetListView.TableRow")
 				.Cursor(EMouseCursor::Default),
 				OwnerTableView);
+		}
 
 
-			TSharedRef< SHorizontalBox > Widget = SNew(SHorizontalBox)
-				+ SHorizontalBox::Slot()
-				.AutoWidth()
-				.Padding(0, 0, 2, 0)
-				.VAlign(VAlign_Center)
-				[
-					// Item icon
-					SNew(SImage)
-					.Image(this, &SAssetPreviewTableRow::GetIconBrush)
-					.ColorAndOpacity(IconColor)
-				]
+		virtual TSharedRef<SWidget> GenerateWidgetForColumn(const FName& ColumnName) override
+		{
+			auto ItemPtr = ItemWeakPtr.Pin();
+			if ( !ItemPtr.IsValid() )
+			{
+				return SNullWidget::NullWidget;
+			}
 
-				+ SHorizontalBox::Slot()
-				.AutoWidth()
-				.VAlign(VAlign_Center)
-				[
-					SNew(STextBlock)
-					.Text(FText::FromString(*(InItem->Name)))
-					.Font(FEditorStyle::GetFontStyle("ContentBrowser.SourceTreeItemFont"))
-					.HighlightText(PreviewWidgetWeakPtr.Pin().Get(), &SAssetsPreviewWidget::OnGetHighlightText)
-				];
+			TSharedRef<SAssetsPreviewWidget> AssetPreview = PreviewWidgetWeakPtr.Pin().ToSharedRef();
+			TSharedRef<SWidget> NewItemWidget = SNullWidget::NullWidget;
 
-			SetContent(Widget);
-		}
+			auto Column = AssetPreview->GetColumn( ColumnName );
+			if ( Column.IsValid() )
+			{
+				NewItemWidget = Column->ConstructRowWidget( ItemPtr.ToSharedRef(), *this, AssetPreview );
+			}
 
-	private:
-		const FSlateBrush* GetIconBrush() const
-		{
-			const FSlateBrush* IconBrush = AssetIconBrush;
-			if (ItemWeakPtr.Pin()->IsFolder())
+			if( ColumnName == FAssetPreviewDefaultColumn::ColumnID )
 			{
-				const bool bExpanded = PreviewWidgetWeakPtr.Pin()->GetTreeView()->IsItemExpanded(ItemWeakPtr.Pin());
-				IconBrush = bExpanded ? FolderOpenBrush : FolderClosedBrush;
+				// The first column gets the tree expansion arrow for this row
+				return
+					SNew( SHorizontalBox )
+
+					+SHorizontalBox::Slot()
+					.AutoWidth()
+					.Padding(6, 0, 0, 0)
+					[
+						SNew( SExpanderArrow, SharedThis(this) ).IndentAmount(12)
+					]
+
+					+SHorizontalBox::Slot()
+					.FillWidth( 1.0f )
+					[
+						NewItemWidget
+					];
+			}
+			else
+			{
+				// Other columns just get widget content -- no expansion arrow needed
+				return NewItemWidget;
 			}
-			return IconBrush;
 		}
 
 	private:
-		/** Brushes for the different folder states */
-		const FSlateBrush* FolderOpenBrush;
-		const FSlateBrush* FolderClosedBrush;
-		const FSlateBrush* AssetIconBrush;
 
-		FAssetTreeItemWeakPtr ItemWeakPtr;
+	private:
+		IAssetTreeItemWeakPtr ItemWeakPtr;
 
 		/** Weak reference back to the preview widget that owns us */
-		TWeakPtr< const SAssetsPreviewWidget > PreviewWidgetWeakPtr;
+		TWeakPtr< SAssetsPreviewWidget > PreviewWidgetWeakPtr;
 	};
 
 	void SAssetsPreviewWidget::Construct(const FArguments& InArgs)
 	{
+		HeaderRow = SNew(SHeaderRow)
+			.Visibility( EVisibility::Visible );
+
+		SetupColumns();
+
 		ChildSlot
 		[
 			SNew(SVerticalBox)
@@ -217,12 +235,12 @@ namespace AssetPreviewWidget
 			]
 
 			+ SVerticalBox::Slot()
-			//.AutoHeight()
-			.Padding(2.f)
+			.Padding( 0, 0, 0, 2 )
 			[
-				SAssignNew(TreeView, STreeView< FAssetTreeItemPtr >)
-				.SelectionMode( ESelectionMode::Single )
-				.TreeItemsSource(&FilteredRootItems)
+				SAssignNew(TreeView, STreeView<IAssetTreeItemPtr>)
+				.SelectionMode(ESelectionMode::Single)
+				.TreeItemsSource(&RootItems)
+				.HeaderRow(HeaderRow)
 				.OnGenerateRow(this, &SAssetsPreviewWidget::MakeRowWidget)
 				.OnSetExpansionRecursive(this, &SAssetsPreviewWidget::OnSetExpansionRecursive)
 				.OnGetChildren(this, &SAssetsPreviewWidget::OnGetChildren)
@@ -233,239 +251,435 @@ namespace AssetPreviewWidget
 
 	void SAssetsPreviewWidget::SetAssetsList(const TArray< TWeakObjectPtr< UObject > >& InAssetsList, const FString& InPathToReplace, const FString& InSubstitutePath)
 	{
-		PathToReplace = InPathToReplace;
+		PathPrefixToRemove = InPathToReplace;
 		SubstitutePath = InSubstitutePath;
-		RootItems.Empty();
-		TMap< FString, FAssetTreeItemPtr > NamesToRootItem;
 
 		// Make sure the root dir is displayed as "Content".
 		// This is more descriptive for the end user.
 		TArray<FString> Tokens;
 		const FString StrContent = TEXT("Content");
-		SubstitutePath.ParseIntoArray(Tokens, TEXT("/"));
-		if (Tokens.Num() > 0 && !Tokens[0].Equals(StrContent))
+		SubstitutePath.ParseIntoArray( Tokens, TEXT("/") );
+		if ( Tokens.Num() > 0 && !Tokens[0].Equals( StrContent ) )
 		{
 			const int32 StartChar = SubstitutePath.Find(Tokens[0]);
 			SubstitutePath.RemoveAt(StartChar, Tokens[0].Len(), false);
 			SubstitutePath.InsertAt(StartChar, StrContent);
 		}
 
-		for (const TWeakObjectPtr< UObject >& Asset : InAssetsList)
+
+		UnFilteredAssets.Empty( InAssetsList.Num() );
+		for ( const TWeakObjectPtr< UObject >& Asset : InAssetsList )
 		{
-			if (Asset.Get())
+			if ( Asset.Get() )
 			{
-				TArray< FString > ItemsName = GetItemsName(Asset);
-
-				if (ItemsName.Num() > 0)
-				{
-					FAssetTreeItemPtr LastParent;
-
-					if (ItemsName.Num() > 1)
-					{
-						// Manage the root item
-						{
-							FString ItemName = MoveTemp(ItemsName[0]);
-							FAssetTreeItemPtr RootItem;
-							if (FAssetTreeItemPtr* PtrToTreeItemPtr = NamesToRootItem.Find(ItemName))
-							{
-								RootItem = *PtrToTreeItemPtr;
-							}
-							else
-							{
-								RootItem = MakeShared<FAssetTreeItem>();
-								RootItem->Name = ItemName;
-								RootItem->OwnerWeakPtr = SharedThis(this);
-								RootItems.Add(RootItem);
-								NamesToRootItem.Add(MoveTemp(ItemName), RootItem);
-							}
-							LastParent = MoveTemp(RootItem);
-						}
-
-						//Manage the folders to the asset
-						{
-							const int32 ItemsNameLenghtMinusOne = ItemsName.Num() - 1;
-							for (int32 i = 1; i < ItemsNameLenghtMinusOne; i++)
-							{
-								FString ItemName = MoveTemp(ItemsName[i]);
-								FAssetTreeItemPtr FolderItem;
-								if (FAssetTreeItemPtr* PtrToTreeItemPtr = LastParent->NameToFolder.Find(ItemName))
-								{
-									FolderItem = *PtrToTreeItemPtr;
-								}
-								else
-								{
-									FolderItem = MakeShared<FAssetTreeItem>();
-									FolderItem->Name = MoveTemp(ItemName);
-									FolderItem->OwnerWeakPtr = SharedThis(this);
-									LastParent->AddFolder(FolderItem);
-								}
-
-								LastParent = FolderItem;
-							}
-						}
-					}
-
-					// Create the asset item
-					{
-						FAssetTreeItemPtr AssetItem = MakeShared<FAssetTreeItem>();
-						AssetItem->Name = MoveTemp(ItemsName.Last());
-						AssetItem->AssetPtr = Asset;
-						AssetItem->OwnerWeakPtr = SharedThis(this);
-						LastParent->Assets.Add(AssetItem);
-					}
-				}
+				FAssetTreeAssetItemPtr AssetItem = MakeShared<FAssetTreeAssetItem>();
+				AssetItem->Name = Asset->GetName();
+				AssetItem->AssetPtr = Asset;
+				UnFilteredAssets.Add( MoveTemp( AssetItem ) );
 			}
 		}
 
-		// Sort items in alphabetical order
-		SortRecursive(RootItems);
-
-		FilterAssetsNames();
+		Refresh();
 	}
 
 	void SAssetsPreviewWidget::ClearAssetList()
 	{
-		RootItems.Empty();
-		FilterAssetsNames();
+		UnFilteredAssets.Empty();
+		Refresh();
+	}
+
+	void SAssetsPreviewWidget::RequestSort()
+	{
+		bIsSortDirty = true;
 	}
 
-	void SAssetsPreviewWidget::FilterAssetsNames()
+	void SAssetsPreviewWidget::Tick(const FGeometry& AllottedGeometry, const double InCurrentTime, const float InDeltaTime)
 	{
-		FilteredRootItems.Empty();
+		if ( bRequestedRefresh )
+		{
+			RootItems.Empty();
+			NameToRootFolder.Empty();
+			CurrentProcessingAssetIndex = 0;
+			bRequestedRefresh = false;
+			TreeView->RequestListRefresh();
+		}
 
-		for (FAssetTreeItemPtr& Item : RootItems)
+		UpdateColumns();
+
+		if ( CurrentProcessingAssetIndex < UnFilteredAssets.Num() )
 		{
-			if (Item->Filter(FilterText))
+			RequestSort();
+
+			const int32 AssetToProccess = FMath::Min<int32>(  UnFilteredAssets.Num() - CurrentProcessingAssetIndex, 500 );
+			int32 Index = 0;
+			while ( Index < AssetToProccess )
 			{
-				FilteredRootItems.Add(Item);
+				FAssetTreeAssetItemPtr& AssetItem = UnFilteredAssets[Index + CurrentProcessingAssetIndex];
+				if ( DoesPassFilter( AssetItem ) )
+				{
+					TArray<FString> Names = GetItemsName( AssetItem->AssetPtr );
+					FAssetTreeFolderItemPtr Parent = FindOrCreateParentsItem( MakeArrayView( Names.GetData(), Names.Num() - 1 ) );
+					Parent->Assets.Add( AssetItem );
+				}
+				++Index;
 			}
+
+			CurrentProcessingAssetIndex += Index;
 		}
 
-		TreeView->RequestListRefresh();
+		if ( bIsSortDirty )
+		{
+			Sort( RootItems );
+			TreeView->RequestListRefresh();
+			ExpandAllFolders();
+			bIsSortDirty = false;
+		}
 
-		ExpandAllFolders();
 	}
 
 	void SAssetsPreviewWidget::ExpandAllFolders()
 	{
-		for (FAssetTreeItemPtr& Item : FilteredRootItems)
+		for ( const IAssetTreeItemPtr& Item : RootItems )
 		{
-			ExpandFolderRecursive(Item);
+			OnSetExpansionRecursive( Item, true );
 		}
 	}
 
-	void SAssetsPreviewWidget::ExpandFolderRecursive(FAssetTreeItemPtr InItem)
+	void SAssetsPreviewWidget::Sort(TArray<IAssetTreeItemPtr>& InItems) const
 	{
-		TreeView->SetItemExpansion(InItem, true);
-		for (FAssetTreeItemPtr& Item : InItem->Folders)
+		if ( const TSharedRef<IAssetPreviewColumn>* ColumnPtr = Columns.Find( SortingColumn ) ) 
 		{
-			ExpandFolderRecursive(Item);
+			const TSharedRef<IAssetPreviewColumn>&  Column = *ColumnPtr;
+			Column->SortItems( InItems, SortingMode );
 		}
 	}
 
-	void SAssetsPreviewWidget::SortRecursive(TArray< FAssetTreeItemPtr >& InItems)
+	FString SAssetsPreviewWidget::GetItemPath(const TWeakObjectPtr<UObject>& Asset) const
 	{
-		InItems.Sort([](const FAssetTreeItemPtr A, const FAssetTreeItemPtr B) { return A->Name.Compare(B->Name) < 0; });
-
-		for (int ItemIndex = 0; ItemIndex < InItems.Num(); ++ItemIndex)
+		FString AssetSubPath = Asset->GetPathName();
+		if ( AssetSubPath.RemoveFromStart( PathPrefixToRemove ) && !SubstitutePath.IsEmpty() )
 		{
-			SortRecursive(InItems[ItemIndex]->Folders);
-			SortRecursive(InItems[ItemIndex]->Assets);
+			AssetSubPath = SubstitutePath / AssetSubPath;
 		}
+		return AssetSubPath;
 	}
 
-	TArray< FString > SAssetsPreviewWidget::GetItemsName(TWeakObjectPtr< UObject > Asset) const
+	TArray< FString > SAssetsPreviewWidget::GetItemsName(const TWeakObjectPtr<UObject>& Asset) const
 	{
-		TArray< FString > ItemsName;
-		FString AssetSubPath = Asset->GetPathName(nullptr);
-		if(AssetSubPath.RemoveFromStart(PathToReplace) && !SubstitutePath.IsEmpty())
-		{
-			AssetSubPath = SubstitutePath / AssetSubPath;
-		}
-		AssetSubPath.ReplaceCharInline(TEXT('/'), TEXT('.'));
-		AssetSubPath.ParseIntoArray(ItemsName, TEXT("."), true);
+		TArray<FString> ItemsName;
+		GetItemPath( Asset ).ParseIntoArray( ItemsName, TEXT("/"), true );
+		return ItemsName;
+	}
+
+	TSharedRef< ITableRow > SAssetsPreviewWidget::MakeRowWidget(IAssetTreeItemPtr InItem, const TSharedRef< STableViewBase >& OwnerTable)
+	{
+		return SNew( SAssetPreviewTableRow, OwnerTable, InItem, SharedThis(this) );
+	}
 
-		// Check that asset's name is not repeated twice at the end
-		if(ItemsName.Num() > 1)
+	void SAssetsPreviewWidget::OnGetChildren(IAssetTreeItemPtr InParent, TArray<IAssetTreeItemPtr>& OutChildren) const
+	{
+		if ( InParent->IsFolder() )
 		{
-			const int32 LastIndex = ItemsName.Num() - 1;
-			if(ItemsName[LastIndex] == ItemsName[LastIndex - 1])
+			FAssetTreeFolderItem& FolderItem = static_cast<FAssetTreeFolderItem&>( *InParent.Get() );
+			OutChildren.Reserve( FolderItem.Folders.Num() + FolderItem.Assets.Num() );
+
+			Sort( FolderItem.Folders );
+
+			for ( IAssetTreeItemPtr& Folder : FolderItem.Folders )
+			{
+				OutChildren.Add( Folder );
+			}
+
+			Sort( FolderItem.Assets );
+
+			for ( IAssetTreeItemPtr& Asset : FolderItem.Assets )
 			{
-				ItemsName.RemoveAt(LastIndex, 1, false);
+				OutChildren.Add( Asset );
 			}
 		}
+	}
 
-		return ItemsName;
+	void SAssetsPreviewWidget::OnSearchBoxChanged(const FText& InSearchText)
+	{
+		FilterText = InSearchText;
+		FilterStrings.Reset();
+		FilterText.ToString().ParseIntoArray( FilterStrings, TEXT(" ")) ;
+		Refresh();
 	}
 
-	TSharedRef< ITableRow > SAssetsPreviewWidget::MakeRowWidget(FAssetTreeItemPtr InItem, const TSharedRef< STableViewBase >& OwnerTable) const
+	void SAssetsPreviewWidget::OnSearchBoxCommitted(const FText& InSearchText, ETextCommit::Type CommitInfo)
 	{
-		TSharedPtr< SAssetPreviewTableRow > TableRowWidget;
+		FilterText = InSearchText;
+		FilterStrings.Empty();
+		FilterText.ToString().ParseIntoArray( FilterStrings, TEXT(" ") );
+		Refresh();
+	}
 
-		SAssignNew(TableRowWidget, SAssetPreviewTableRow, OwnerTable, InItem, SharedThis(this));
+	FText SAssetsPreviewWidget::OnGetHighlightText() const
+	{
+		return FilterText;
+	}
 
-		return TableRowWidget.ToSharedRef();
+	TSharedPtr<IAssetPreviewColumn> SAssetsPreviewWidget::GetColumn(FName ColumnID) const
+	{
+		if ( const TSharedRef<IAssetPreviewColumn>* Column = Columns.Find(ColumnID))
+		{
+			return *Column;
+		}
+
+		return {};
 	}
 
-	void SAssetsPreviewWidget::OnGetChildren(FAssetTreeItemPtr InParent, TArray< FAssetTreeItemPtr >& OutChildren) const
+	void SAssetsPreviewWidget::AddColumn(TSharedRef<IAssetPreviewColumn> Column)
 	{
-		for (FAssetTreeItemPtr& Folder : InParent->Folders)
+		PendingColumnsToAdd.Add( Column );
+	}
+
+	void SAssetsPreviewWidget::RemoveColumn(FName ColumnID)
+	{
+		PendingColumnsToRemove.Add( ColumnID );
+	}
+
+	void SAssetsPreviewWidget::OnSetExpansionRecursive(IAssetTreeItemPtr InTreeNode, bool bInIsItemExpanded)
+	{
+		if ( InTreeNode.IsValid() && InTreeNode->IsFolder() )
 		{
-			if (Folder->bPassedFilter)
+			TreeView->SetItemExpansion( InTreeNode, bInIsItemExpanded );
+			FAssetTreeFolderItem& Folder = static_cast<FAssetTreeFolderItem&>( *InTreeNode.Get() );
+			for ( const IAssetTreeItemPtr& SubFolder : Folder.Folders )
 			{
-				OutChildren.Add(Folder);
+				OnSetExpansionRecursive( SubFolder, bInIsItemExpanded );
 			}
 		}
+	}
 
-		for (FAssetTreeItemPtr& Asset : InParent->Assets)
+	void SAssetsPreviewWidget::OnSelectionChangedInternal(IAssetTreeItemPtr ItemSelected, ESelectInfo::Type SelectionType)
+	{
+		if ( ItemSelected )
 		{
-			if (Asset->bPassedFilter)
+			TSet< UObject* > Selection;
+			if ( !ItemSelected->IsFolder() )
 			{
-				OutChildren.Add(Asset);
+				Selection.Add( static_cast<FAssetTreeAssetItem&>( *ItemSelected.Get() ).AssetPtr.Get() );
 			}
+			OnSelectionChanged().Broadcast( Selection );
 		}
+	}
 
+	EColumnSortMode::Type SAssetsPreviewWidget::GetColumnSortMode(const FName ColumnId) const
+	{
+		if ( SortingColumn == ColumnId )
+		{
+			return SortingMode;
+		}
+
+		return EColumnSortMode::None;
 	}
 
-	void SAssetsPreviewWidget::OnSearchBoxChanged(const FText& InSearchText)
+	void SAssetsPreviewWidget::OnColumnSortModeChanged(const EColumnSortPriority::Type SortPriority, const FName& ColumnId, const EColumnSortMode::Type InSortMode)
 	{
-		FilterText = InSearchText;
-		FilterAssetsNames();
+		if ( SortPriority == EColumnSortPriority::Primary )
+		{
+			SortingColumn = ColumnId;
+			SortingMode = InSortMode;
+			RequestSort();
+		}
 	}
 
-	void SAssetsPreviewWidget::OnSearchBoxCommitted(const FText& InSearchText, ETextCommit::Type CommitInfo)
+	void SAssetsPreviewWidget::SetupColumns()
 	{
-		FilterText = InSearchText;
-		FilterAssetsNames();
+		SHeaderRow* HeaderRowPtr = HeaderRow.Get();
+		check( HeaderRowPtr );
+
+		Columns.Add( AssetPreviewWidget::FAssetPreviewDefaultColumn::ColumnID, MakeShared<AssetPreviewWidget::FAssetPreviewDefaultColumn>() );
+
+		SortingColumn = AssetPreviewWidget::FAssetPreviewDefaultColumn::ColumnID;
+		SortingMode = EColumnSortMode::Ascending;
+
+		TSharedRef<SAssetsPreviewWidget> AssetPreviewWidget = StaticCastSharedRef<SAssetsPreviewWidget>( AsShared() );
+		for (const TPair<FName,TSharedRef<IAssetPreviewColumn>>& Pair : Columns )
+		{
+			HeaderRowPtr->AddColumn(
+				Pair.Value->ConstructHeaderRowColumn( AssetPreviewWidget )
+					.SortMode(this, &SAssetsPreviewWidget::GetColumnSortMode, AssetPreviewWidget::FAssetPreviewDefaultColumn::ColumnID)
+					.OnSort(this, &SAssetsPreviewWidget::OnColumnSortModeChanged)
+				);
+		}
 	}
 
-	FText SAssetsPreviewWidget::OnGetHighlightText() const
+	void SAssetsPreviewWidget::UpdateColumns()
 	{
-		return FilterText;
+		SHeaderRow* HeaderRowPtr = HeaderRow.Get();
+		check( HeaderRowPtr );
+
+		if ( PendingColumnsToRemove.Num() > 0 )
+		{
+			for ( const FName& ColumnToRemove : PendingColumnsToRemove )
+			{
+				HeaderRowPtr->RemoveColumn( ColumnToRemove );
+				Columns.Remove( ColumnToRemove );
+			}
+			PendingColumnsToRemove.Empty();
+		}
+
+
+		if ( PendingColumnsToAdd.Num() > 0 ) 
+		{
+			PendingColumnsToAdd.StableSort([](const TSharedRef<IAssetPreviewColumn>& First, const TSharedRef<IAssetPreviewColumn>& Second) -> bool
+				{
+					return First->GetCulumnPositionPriorityIndex() >= Second->GetCulumnPositionPriorityIndex();
+				});
+
+			TSharedRef<SAssetsPreviewWidget> AssetPreviewWidget = StaticCastSharedRef<SAssetsPreviewWidget>( AsShared() );
+
+			int32 CurrentElementToAddIndex = 0;
+			int32 CurrentPriority = PendingColumnsToAdd[CurrentElementToAddIndex]->GetCulumnPositionPriorityIndex();
+			int32 CurrentColumnIndex = 0;
+			TArray<int32> InsertionIndex;
+			InsertionIndex.Reserve( PendingColumnsToAdd.Num() );
+
+			for (const auto& Pair : Columns)
+			{
+				if ( Pair.Value->GetCulumnPositionPriorityIndex() < CurrentPriority )
+				{
+					InsertionIndex.Add( CurrentElementToAddIndex + CurrentColumnIndex );
+					++CurrentElementToAddIndex;
+					if ( CurrentElementToAddIndex >= PendingColumnsToAdd.Num() )
+					{
+						break;
+					}
+					CurrentPriority = PendingColumnsToAdd[CurrentElementToAddIndex]->GetCulumnPositionPriorityIndex();
+				}
+
+				++CurrentColumnIndex;
+			}
+
+			
+
+			Columns.Reserve( PendingColumnsToAdd.Num() );
+			for ( TSharedRef<IAssetPreviewColumn>& NewColumn : PendingColumnsToAdd )
+			{
+				Columns.Add( NewColumn->GetColumnID(), NewColumn );
+			}
+
+			for (int32 Index = 0; Index < InsertionIndex.Num(); ++Index)
+			{
+				HeaderRowPtr->InsertColumn(
+					PendingColumnsToAdd[Index]->ConstructHeaderRowColumn( AssetPreviewWidget )
+						.SortMode( this, &SAssetsPreviewWidget::GetColumnSortMode, PendingColumnsToAdd[CurrentElementToAddIndex]->GetColumnID() )
+						.OnSort( this, &SAssetsPreviewWidget::OnColumnSortModeChanged )
+					, InsertionIndex[Index]
+					);
+			}
+
+			while (CurrentElementToAddIndex < PendingColumnsToAdd.Num())
+			{
+				HeaderRowPtr->AddColumn(
+					PendingColumnsToAdd[CurrentElementToAddIndex]->ConstructHeaderRowColumn( AssetPreviewWidget )
+					.SortMode( this, &SAssetsPreviewWidget::GetColumnSortMode, PendingColumnsToAdd[CurrentElementToAddIndex]->GetColumnID() )
+					.OnSort( this, &SAssetsPreviewWidget::OnColumnSortModeChanged )
+					);
+				++CurrentElementToAddIndex;
+			}
+
+			PendingColumnsToAdd.Empty();
+
+			Columns.ValueSort([](const TSharedRef<IAssetPreviewColumn>& First, const TSharedRef<IAssetPreviewColumn>& Second) -> bool
+				{
+					return First->GetCulumnPositionPriorityIndex() >= Second->GetCulumnPositionPriorityIndex();
+				}
+			);
+		}
 	}
 
+	void SAssetsPreviewWidget::Refresh()
+	{
+		bRequestedRefresh = true;
+	}
 
-	void SAssetsPreviewWidget::OnSetExpansionRecursive(FAssetTreeItemPtr InTreeNode, bool bInIsItemExpanded)
+	FAssetTreeFolderItemPtr SAssetsPreviewWidget::FindOrCreateParentsItem(const TArrayView<FString>& ParentNames)
 	{
-		if (InTreeNode.IsValid())
+		if ( ParentNames.Num() > 0 )
 		{
-			TreeView->SetItemExpansion(InTreeNode, bInIsItemExpanded);
+			FAssetTreeFolderItemPtr LastParent = nullptr;
+			FString* CurrentParentName = &ParentNames[0];
+			uint32 CurrentHash = GetTypeHash( *CurrentParentName );
 
-			for (FAssetTreeItemPtr SubFolder : InTreeNode->Folders)
+			// We are starting from the root
+			if ( FAssetTreeFolderItemPtr* RootPtr = NameToRootFolder.FindByHash( CurrentHash, *CurrentParentName ) )
+			{
+				LastParent = *RootPtr;
+			}
+			else
 			{
-				OnSetExpansionRecursive(SubFolder, bInIsItemExpanded);
+				FAssetTreeFolderItemPtr Root = MakeShared<FAssetTreeFolderItem>();
+				Root->Name = *CurrentParentName;
+				LastParent = Root;
+				RootItems.Add( Root );
+				NameToRootFolder.AddByHash( CurrentHash, *CurrentParentName, MoveTemp( Root ) );
 			}
+
+			check( LastParent );
+
+			for ( int32 Index = 1; Index < ParentNames.Num(); ++Index )
+			{
+				CurrentParentName = &ParentNames[Index];
+				CurrentHash = GetTypeHash( *CurrentParentName );
+				if ( FAssetTreeFolderItemPtr* CurrentPtr = LastParent->NameToFolder.FindByHash( CurrentHash, *CurrentParentName ) )
+				{
+					LastParent = *CurrentPtr;
+				}
+				else
+				{
+					FAssetTreeFolderItemPtr CurrentParent = MakeShared<FAssetTreeFolderItem>();
+					CurrentParent->Name = *CurrentParentName;
+					LastParent->Folders.Add( CurrentParent );
+					LastParent->NameToFolder.AddByHash( CurrentHash, *CurrentParentName, CurrentParent );
+					LastParent = MoveTemp( CurrentParent );
+				}
+			}
+
+			check( LastParent );
+			return LastParent;
 		}
+
+		return {};
 	}
 
-	void SAssetsPreviewWidget::OnSelectionChangedInternal(FAssetTreeItemPtr ItemSelected, ESelectInfo::Type SelectionType)
+	bool SAssetsPreviewWidget::DoesPassFilter(const FAssetTreeAssetItemPtr& AssetItem) const
 	{
-		if ( ItemSelected )
-		{	
-			TSet< UObject* > Selection;
-			Selection.Add( ItemSelected->AssetPtr.Get() );
-			OnSelectionChanged().Broadcast( Selection );
+		if ( FilterStrings.Num() == 0 )
+		{
+			return true;
+		}
+
+		TArray<FString> SearchableStrings;
+		for ( const TPair<FName, TSharedRef<IAssetPreviewColumn>>& Pair : Columns )
+		{
+			Pair.Value->PopulateSearchStrings( AssetItem, SearchableStrings, *this );
 		}
+		
+		// All key word must match against a least one of the provided string
+		bool bPassKeyWord = false;
+		for ( const FString& KeyWord : FilterStrings )
+		{
+			bPassKeyWord = false;
+			for (const FString& String : SearchableStrings)
+			{
+				if ( String.Contains(KeyWord) )
+				{
+					bPassKeyWord = true;
+					break;
+				}
+			}
+
+			if ( !bPassKeyWord )
+			{
+				break;
+			}
+		}
+
+		return bPassKeyWord;
 	}
 }
 
diff --git a/Engine/Plugins/Enterprise/DataprepEditor/Source/DataprepEditor/Private/Widgets/SAssetsPreviewWidget.h b/Engine/Plugins/Enterprise/DataprepEditor/Source/DataprepEditor/Private/Widgets/SAssetsPreviewWidget.h
--- a/Engine/Plugins/Enterprise/DataprepEditor/Source/DataprepEditor/Private/Widgets/SAssetsPreviewWidget.h
+++ b/Engine/Plugins/Enterprise/DataprepEditor/Source/DataprepEditor/Private/Widgets/SAssetsPreviewWidget.h
@@ -4,6 +4,7 @@
 
 #include "CoreMinimal.h"
 
+#include "Templates/SharedPointer.h"
 #include "Types/SlateEnums.h"
 #include "Widgets/SCompoundWidget.h"
 #include "Widgets/SWidget.h"
@@ -11,13 +12,51 @@
 #include "Widgets/Views/STableViewBase.h"
 #include "Widgets/Views/STreeView.h"
 
+class SHeaderRow;
 
 namespace AssetPreviewWidget
 {
-	DECLARE_MULTICAST_DELEGATE_OneParam(FOnSelectionChanged, TSet< UObject* > /** Selected objects */)
+	struct IAssetTreeItem;
+
+	using IAssetTreeItemPtr = TSharedPtr<IAssetTreeItem>;
+	using IAssetTreeItemWeakPtr = TWeakPtr<IAssetTreeItem>;
+
+	struct FAssetTreeAssetItem;
+	using FAssetTreeAssetItemPtr = TSharedPtr<FAssetTreeAssetItem>;
+
+	struct FAssetTreeFolderItem;
+	using FAssetTreeFolderItemPtr = TSharedPtr<FAssetTreeFolderItem>;
 
-	using FAssetTreeItemPtr = TSharedPtr< struct FAssetTreeItem >;
-	using FAssetTreeItemWeakPtr = TWeakPtr< struct FAssetTreeItem >;
+
+	class SAssetsPreviewWidget;
+
+	/**
+	 *	Interface for a column of the asset preview
+	 */
+	class IAssetPreviewColumn : public TSharedFromThis< IAssetPreviewColumn >
+	{
+	public:
+		virtual ~IAssetPreviewColumn() = default;
+
+		// A index for the desired position column. The higher index end up on left.
+		virtual uint8 GetCulumnPositionPriorityIndex() const = 0;
+
+		virtual FName GetColumnID() const = 0;
+
+		virtual SHeaderRow::FColumn::FArguments ConstructHeaderRowColumn(const TSharedRef<SAssetsPreviewWidget>& PreviewWidget) = 0;
+
+		virtual const TSharedRef< SWidget > ConstructRowWidget(const IAssetTreeItemPtr& TreeItem, const STableRow<IAssetTreeItemPtr>& Row, const TSharedRef<SAssetsPreviewWidget>& PreviewWidget) = 0;
+
+	public:
+		/** Optionally overridden interface methods */
+		virtual void PopulateSearchStrings(const IAssetTreeItemPtr& Item, TArray< FString >& OutSearchStrings, const SAssetsPreviewWidget& AssetPreview) const {}
+
+		virtual void SortItems(TArray<IAssetTreeItemPtr>& OutItems, const EColumnSortMode::Type SortMode) const {}
+	};
+
+	using FAssetPreviewColumn = TSharedRef<IAssetPreviewColumn>;
+
+	DECLARE_MULTICAST_DELEGATE_OneParam(FOnSelectionChanged, TSet< UObject* > /** Selected objects */)
 
 	class SAssetsPreviewWidget : public SCompoundWidget
 	{
@@ -31,46 +70,134 @@ namespace AssetPreviewWidget
 		void SetAssetsList(const TArray< TWeakObjectPtr< UObject > >& InAssetsList, const FString& InPathToReplace, const FString& InSubstitutePath);
 		void ClearAssetList();
 
+		void RequestSort();
+
 		FOnSelectionChanged& OnSelectionChanged() { return OnSelectionChangedDelegate; };
 	
 		FText OnGetHighlightText() const;
 
-		const TSharedPtr< STreeView< FAssetTreeItemPtr > > GetTreeView() const
+		const TSharedPtr< STreeView< IAssetTreeItemPtr > > GetTreeView() const
 		{
 			return TreeView;
 		}
 
+		TSharedPtr<IAssetPreviewColumn> GetColumn(FName ColumnID) const;
+
+		void AddColumn(TSharedRef<IAssetPreviewColumn> Column);
+
+		void RemoveColumn(FName ColumnID);
+
+		// Begin of SWidgetInterface
+		void Tick(const FGeometry& AllottedGeometry, const double InCurrentTime, const float InDeltaTime) override;
+		// End of SWidgetInterface
+
+		TArray<FString> GetItemsName(const TWeakObjectPtr<UObject>& Asset) const;
+
+		EColumnSortMode::Type GetColumnSortMode(const FName ColumnId) const;
+
 	private:
 
 		friend struct FAssetTreeItem;
 
-		void FilterAssetsNames();
 		void ExpandAllFolders();
-		void ExpandFolderRecursive(FAssetTreeItemPtr InItem);
 
-		void SortRecursive(TArray< FAssetTreeItemPtr >& InItems);
+		void Sort(TArray<IAssetTreeItemPtr>& InItems) const;
 
-		TArray< FString > GetItemsName(TWeakObjectPtr< UObject > Asset) const;
+		FString GetItemPath(const TWeakObjectPtr<UObject>& Asset) const;
 
-		TSharedRef< class ITableRow > MakeRowWidget(FAssetTreeItemPtr InItem, const TSharedRef< class STableViewBase >& OwnerTable) const;
-		void OnGetChildren(FAssetTreeItemPtr InParent, TArray< FAssetTreeItemPtr >& OutChildren) const;
+		TSharedRef< class ITableRow > MakeRowWidget(IAssetTreeItemPtr InItem, const TSharedRef< class STableViewBase >& OwnerTable);
+		void OnGetChildren(IAssetTreeItemPtr InParent, TArray<IAssetTreeItemPtr>& OutChildren) const;
 
 		void OnSearchBoxChanged(const FText& InSearchText);
 		void OnSearchBoxCommitted(const FText& InSearchText, ETextCommit::Type CommitInfo);
 
-		void OnSetExpansionRecursive(FAssetTreeItemPtr InTreeNode, bool bInIsItemExpanded);
+		void OnSetExpansionRecursive(const IAssetTreeItemPtr InTreeNode, bool bInIsItemExpanded);
+
+		void OnSelectionChangedInternal(IAssetTreeItemPtr ItemSelected, ESelectInfo::Type SelectionType);
+
+		void OnColumnSortModeChanged(const EColumnSortPriority::Type SortPriority, const FName& ColumnId, const EColumnSortMode::Type InSortMode);
+
 
-		void OnSelectionChangedInternal(FAssetTreeItemPtr ItemSelected, ESelectInfo::Type SelectionType);
+		void SetupColumns();
 
-		TArray< FAssetTreeItemPtr > RootItems;
-		TArray< FAssetTreeItemPtr > FilteredRootItems;
+		void UpdateColumns();
 
-		TSharedPtr< STreeView< FAssetTreeItemPtr > > TreeView;
+		void Refresh();
+
+		/**
+		 * Find or create the parents chain
+		 * @param ParentNames the names of the parents for the top to bottom
+		 * @return the immediate parent
+		 */
+		FAssetTreeFolderItemPtr FindOrCreateParentsItem(const TArrayView<FString>& ParentNames);
+
+		// For now we only filter the assets items
+		bool DoesPassFilter(const FAssetTreeAssetItemPtr& AssetItem) const;
+
+		// The root items displayed
+		TArray<IAssetTreeItemPtr> RootItems;
+		TMap<FString, FAssetTreeFolderItemPtr> NameToRootFolder;
+
+		// The list of asset unfiltered
+		TArray<FAssetTreeAssetItemPtr> UnFilteredAssets;
+		int32 CurrentProcessingAssetIndex = 0;
+
+		TSharedPtr< STreeView< IAssetTreeItemPtr > > TreeView;
+		TSharedPtr< SHeaderRow > HeaderRow;
 
 		FText FilterText;
-		FString PathToReplace;
+		FString PathPrefixToRemove;
 		FString SubstitutePath;
 
+		TMap<FName, TSharedRef<IAssetPreviewColumn>> Columns;
+
+		TArray<TSharedRef<IAssetPreviewColumn>> PendingColumnsToAdd;
+		TArray<FName> PendingColumnsToRemove;
+
 		FOnSelectionChanged OnSelectionChangedDelegate;
+
+		FName SortingColumn;
+		EColumnSortMode::Type SortingMode;
+
+		// The string used for the filtering
+		TArray<FString> FilterStrings;
+
+		bool bIsSortDirty = false;
+		bool bRequestedRefresh = false;
+	};
+
+
+	struct IAssetTreeItem
+	{
+		virtual ~IAssetTreeItem() = default;
+
+		FString Name;
+		virtual bool IsFolder() const  = 0;
 	};
+
+	struct FAssetTreeAssetItem : public IAssetTreeItem
+	{
+		virtual bool IsFolder() const final
+		{
+			return false;
+		}
+
+		TWeakObjectPtr< UObject > AssetPtr;
+	};
+
+	struct FAssetTreeFolderItem : public IAssetTreeItem
+	{
+		virtual bool IsFolder() const final
+		{
+			return true;
+		}
+
+		// This is used to accelerate the construction of the tree in the set assets function
+		TMap<FString, FAssetTreeFolderItemPtr> NameToFolder;
+
+		// Childrens
+		TArray<IAssetTreeItemPtr> Folders;
+		TArray<IAssetTreeItemPtr> Assets;
+	};
+
 }
diff --git a/Engine/Plugins/Enterprise/DataprepEditor/Source/DataprepEditor/Private/Widgets/SDataprepPreviewRow.cpp b/Engine/Plugins/Enterprise/DataprepEditor/Source/DataprepEditor/Private/Widgets/SDataprepPreviewRow.cpp
new file mode 100644
--- /dev/null
+++ b/Engine/Plugins/Enterprise/DataprepEditor/Source/DataprepEditor/Private/Widgets/SDataprepPreviewRow.cpp
@@ -0,0 +1,82 @@
+// Copyright Epic Games, Inc. All Rights Reserved.
+
+#include "Widgets/SDataprepPreviewRow.h"
+
+#include "DataprepEditorStyle.h"
+#include "PreviewSystem/DataprepPreviewSystem.h"
+
+#include "EditorFontGlyphs.h"
+#include "EditorStyleSet.h"
+#include "Internationalization/Text.h"
+#include "Widgets/Layout/SBox.h"
+#include "Widgets/SBoxPanel.h"
+#include "Widgets/Text/STextBlock.h"
+
+void SDataprepPreviewRow::Construct(const FArguments& InArgs, const TSharedPtr<FDataprepPreviewProcessingResult>& InPreviewData)
+{
+	PreviewData = InPreviewData;
+
+	ChildSlot
+	[
+		SNew( SHorizontalBox )
+		+ SHorizontalBox::Slot()
+		.VAlign( VAlign_Center )
+		.HAlign( HAlign_Left )
+		.AutoWidth()
+		[
+			SNew( SBox )
+			.WidthOverride( 18.f )
+			[
+				SNew( STextBlock )
+				.Font( FEditorStyle::Get().GetFontStyle("FontAwesome.11") )
+				.ColorAndOpacity( FDataprepEditorStyle::GetColor("Graph.ActionStepNode.PreviewColor"))
+				.Text( this, &SDataprepPreviewRow::GetIcon )
+			]
+		]
+		+ SHorizontalBox::Slot()
+		.VAlign(VAlign_Center)
+		.HAlign(HAlign_Left)
+		[
+			SNew( STextBlock )
+			.Text( this, &SDataprepPreviewRow::GetLabel )
+			.HighlightText( InArgs._HighlightText )
+		]
+	];
+}
+
+FText SDataprepPreviewRow::GetIcon() const
+{
+	if ( FDataprepPreviewProcessingResult* Result = PreviewData.Get() )
+	{
+		switch ( Result->Status )
+		{
+		case EDataprepPreviewStatus::BeingProcessed:
+			return FEditorFontGlyphs::Refresh;
+			break;
+		case EDataprepPreviewStatus::Pass:
+			return FEditorFontGlyphs::Check;
+			break;
+		case EDataprepPreviewStatus::Failed:
+			return FText::GetEmpty();
+			break;
+		case EDataprepPreviewStatus::NotSupported :
+			return FEditorFontGlyphs::Question;
+			break;
+		default:
+			break;
+		}
+	}
+
+	return FEditorFontGlyphs::Bug;
+}
+
+FText SDataprepPreviewRow::GetLabel() const
+{
+	if ( FDataprepPreviewProcessingResult* Result = PreviewData.Get() )
+	{
+		 return Result->GetFetchedDataAsText();
+	}
+
+	return {};
+}
+
diff --git a/Engine/Plugins/Enterprise/DataprepEditor/Source/DataprepEditor/Private/Widgets/SDataprepPreviewRow.h b/Engine/Plugins/Enterprise/DataprepEditor/Source/DataprepEditor/Private/Widgets/SDataprepPreviewRow.h
new file mode 100644
--- /dev/null
+++ b/Engine/Plugins/Enterprise/DataprepEditor/Source/DataprepEditor/Private/Widgets/SDataprepPreviewRow.h
@@ -0,0 +1,29 @@
+// Copyright Epic Games, Inc. All Rights Reserved.
+
+#pragma once
+
+#include "CoreMinimal.h"
+#include "Widgets/SCompoundWidget.h"
+#include "Widgets/DeclarativeSyntaxSupport.h"
+
+struct FDataprepPreviewProcessingResult; 
+
+/**
+ * The widget to display the result of the filter preview
+ */
+class SDataprepPreviewRow : public SCompoundWidget
+{
+public:
+	SLATE_BEGIN_ARGS(SDataprepPreviewRow) {}
+		SLATE_ATTRIBUTE(FText, HighlightText)
+	SLATE_END_ARGS();
+
+	void Construct(const FArguments& InArgs, const TSharedPtr<FDataprepPreviewProcessingResult>& InPreviewData);
+
+private:
+	FText GetIcon() const;
+
+	FText GetLabel() const;
+
+	TSharedPtr<FDataprepPreviewProcessingResult> PreviewData;
+};
\ No newline at end of file
diff --git a/Engine/Plugins/Enterprise/DataprepEditor/Source/DataprepLibraries/Private/Fetchers/DataprepIntegerFetcherLibrary.cpp b/Engine/Plugins/Enterprise/DataprepEditor/Source/DataprepLibraries/Private/Fetchers/DataprepIntegerFetcherLibrary.cpp
--- a/Engine/Plugins/Enterprise/DataprepEditor/Source/DataprepLibraries/Private/Fetchers/DataprepIntegerFetcherLibrary.cpp
+++ b/Engine/Plugins/Enterprise/DataprepEditor/Source/DataprepLibraries/Private/Fetchers/DataprepIntegerFetcherLibrary.cpp
@@ -13,7 +13,7 @@
 
 /* UDataprepTriangleCountFetcher methods+
  *****************************************************************************/
-int UDataprepTriangleCountFetcher::Fetch_Implementation(const UObject* Object, bool& bOutFetchSucceded) const
+int32 UDataprepTriangleCountFetcher::Fetch_Implementation(const UObject* Object, bool& bOutFetchSucceded) const
 {
 	bOutFetchSucceded = false;
 
@@ -71,7 +71,7 @@ FText UDataprepTriangleCountFetcher::GetNodeDisplayFetcherName_Implementation()
 
 /* UDataprepVertexCountFetcher methods+
  *****************************************************************************/
-int UDataprepVertexCountFetcher::Fetch_Implementation(const UObject* Object, bool& bOutFetchSucceded) const
+int32 UDataprepVertexCountFetcher::Fetch_Implementation(const UObject* Object, bool& bOutFetchSucceded) const
 {
 	bOutFetchSucceded = false;
 
diff --git a/Engine/Plugins/Enterprise/DataprepEditor/Source/DataprepLibraries/Public/Fetchers/DataprepIntegerFetcherLibrary.h b/Engine/Plugins/Enterprise/DataprepEditor/Source/DataprepLibraries/Public/Fetchers/DataprepIntegerFetcherLibrary.h
--- a/Engine/Plugins/Enterprise/DataprepEditor/Source/DataprepLibraries/Public/Fetchers/DataprepIntegerFetcherLibrary.h
+++ b/Engine/Plugins/Enterprise/DataprepEditor/Source/DataprepLibraries/Public/Fetchers/DataprepIntegerFetcherLibrary.h
@@ -14,7 +14,7 @@ class UDataprepTriangleCountFetcher final : public UDataprepIntegerFetcher
 	GENERATED_BODY()
 public:
 	//~ UDataprepIntegerFetcher interface
-	virtual int Fetch_Implementation(const UObject* Object, bool& bOutFetchSucceded) const final;
+	virtual int32 Fetch_Implementation(const UObject* Object, bool& bOutFetchSucceded) const final;
 	//~ End of UDataprepIntegerFetcher interface
 
 	//~ UDataprepFetcher interface
@@ -29,7 +29,7 @@ class UDataprepVertexCountFetcher final : public UDataprepIntegerFetcher
 	GENERATED_BODY()
 public:
 	//~ UDataprepIntegerFetcher interface
-	virtual int Fetch_Implementation(const UObject* Object, bool& bOutFetchSucceded) const final;
+	virtual int32 Fetch_Implementation(const UObject* Object, bool& bOutFetchSucceded) const final;
 	//~ End of UDataprepIntegerFetcher interface
 
 	//~ UDataprepFetcher interface
diff --git a/Engine/Source/Editor/SceneOutliner/Public/ISceneOutliner.h b/Engine/Source/Editor/SceneOutliner/Public/ISceneOutliner.h
--- a/Engine/Source/Editor/SceneOutliner/Public/ISceneOutliner.h
+++ b/Engine/Source/Editor/SceneOutliner/Public/ISceneOutliner.h
@@ -6,6 +6,7 @@
 #include "Misc/Attribute.h"
 #include "SceneOutlinerFwd.h"
 #include "Widgets/SCompoundWidget.h"
+#include "Widgets/Views/SHeaderRow.h"
 
 // Forward declaration
 template<typename ItemType> class STreeView;
@@ -86,6 +87,12 @@ public:
 	/** Return the name/Id of the columns of the scene outliner */
 	virtual TArray<FName> GetColumnIds() const = 0;
 
+	/** Return the sorting mode for the specified ColumnId */
+	virtual EColumnSortMode::Type GetColumnSortMode( const FName ColumnId ) const = 0;
+
+	/** Request that the tree be sorted at a convenient time */
+	virtual void RequestSort() = 0;
+
 	/** Returns true if edit delete can be executed */
 	virtual bool Delete_CanExecute() = 0;
 
diff --git a/Engine/Source/Editor/SceneOutliner/Public/SSceneOutliner.h b/Engine/Source/Editor/SceneOutliner/Public/SSceneOutliner.h
--- a/Engine/Source/Editor/SceneOutliner/Public/SSceneOutliner.h
+++ b/Engine/Source/Editor/SceneOutliner/Public/SSceneOutliner.h
@@ -191,6 +191,12 @@ namespace SceneOutliner
 		/** Return the name/Id of the columns of the scene outliner */
 		virtual TArray<FName> GetColumnIds() const override;
 
+		/** @return Returns the current sort mode of the specified column */
+		virtual EColumnSortMode::Type GetColumnSortMode(const FName ColumnId) const;
+
+		/** Request that the tree be sorted at a convenient time */
+		virtual void RequestSort();
+
 		/** Returns true if edit delete can be executed */
 		virtual bool Delete_CanExecute();
 
@@ -926,12 +932,6 @@ namespace SceneOutliner
 		/** Handles column sorting mode change */
 		void OnColumnSortModeChanged( const EColumnSortPriority::Type SortPriority, const FName& ColumnId, const EColumnSortMode::Type InSortMode );
 
-		/** @return Returns the current sort mode of the specified column */
-		EColumnSortMode::Type GetColumnSortMode( const FName ColumnId ) const;
-
-		/** Request that the tree be sorted at a convenient time */
-		void RequestSort();
-
 		/** Sort the specified array of items based on the current sort column */
 		void SortItems(TArray<FTreeItemPtr>& Items) const;
 
