DataPrep - Output to Folder should do a snitize check on the name of the folder

diff --git a/Engine/Plugins/Enterprise/DataprepEditor/Source/DataprepLibraries/Private/DataprepLibrariesModule.cpp b/Engine/Plugins/Enterprise/DataprepEditor/Source/DataprepLibraries/Private/DataprepLibrariesModule.cpp
--- a/Engine/Plugins/Enterprise/DataprepEditor/Source/DataprepLibraries/Private/DataprepLibrariesModule.cpp
+++ b/Engine/Plugins/Enterprise/DataprepEditor/Source/DataprepLibraries/Private/DataprepLibrariesModule.cpp
@@ -14,6 +14,7 @@ public:
 	{
 		FPropertyEditorModule& PropertyModule = FModuleManager::LoadModuleChecked< FPropertyEditorModule >( TEXT("PropertyEditor") );
 		PropertyModule.RegisterCustomClassLayout( UDataprepSetLODGroupOperation::StaticClass()->GetFName(), FOnGetDetailCustomizationInstance::CreateStatic( &FDataprepSetLOGGroupDetails::MakeDetails ) );
+		PropertyModule.RegisterCustomClassLayout( UDataprepSetOutputFolder::StaticClass()->GetFName(), FOnGetDetailCustomizationInstance::CreateStatic( &FDataprepSetOutputFolderDetails::MakeDetails ) );
 	}
 
 	virtual void ShutdownModule() override
diff --git a/Engine/Plugins/Enterprise/DataprepEditor/Source/DataprepLibraries/Private/DataprepOperations.cpp b/Engine/Plugins/Enterprise/DataprepEditor/Source/DataprepLibraries/Private/DataprepOperations.cpp
--- a/Engine/Plugins/Enterprise/DataprepEditor/Source/DataprepLibraries/Private/DataprepOperations.cpp
+++ b/Engine/Plugins/Enterprise/DataprepEditor/Source/DataprepLibraries/Private/DataprepOperations.cpp
@@ -7,6 +7,8 @@
 #include "GenericPlatform/GenericPlatformTime.h"
 #include "StaticMeshResources.h"
 
+#include "Misc/FileHelper.h"
+
 // UI related section
 #include "DetailCategoryBuilder.h"
 #include "DetailLayoutBuilder.h"
@@ -392,5 +394,95 @@ void UDataprepSetOutputFolder::OnExecution_Implementation(const FDataprepContext
 	UDataprepOperationsLibrary::SetSubOuputFolder(InContext.Objects, FolderName);
 }
 
+void FDataprepSetOutputFolderDetails::CustomizeDetails(IDetailLayoutBuilder& DetailBuilder)
+{
+	TArray< TWeakObjectPtr< UObject > > Objects;
+	DetailBuilder.GetObjectsBeingCustomized( Objects );
+	check( Objects.Num() > 0 );
+
+	Operation = Cast< UDataprepSetOutputFolder >(Objects[0].Get());
+	check( Operation );
+
+	// #ueent_todo: Remove handling of warning category when this is not considered experimental anymore
+	TArray<FName> CategoryNames;
+	DetailBuilder.GetCategoryNames( CategoryNames );
+	CategoryNames.Remove( FName(TEXT("Warning")) );
+
+	DetailBuilder.HideCategory(FName( TEXT( "Warning" ) ) );
+
+	IDetailCategoryBuilder& CategoryBuilder = DetailBuilder.EditCategory( NAME_None, FText::GetEmpty(), ECategoryPriority::Important );
+
+	FolderNamePropertyHandle = DetailBuilder.GetProperty(GET_MEMBER_NAME_CHECKED(UDataprepSetOutputFolder, FolderName));
+	FolderNamePropertyHandle->MarkHiddenByCustomization();
+
+	FDetailWidgetRow& CustomAssetImportRow = CategoryBuilder.AddCustomRow( FText::FromString( TEXT( "Folder Name" ) ) );
+
+	CustomAssetImportRow.NameContent()
+	[
+		FolderNamePropertyHandle->CreatePropertyNameWidget()
+	];
+
+	CustomAssetImportRow.ValueContent()
+	[
+		SAssignNew(TextBox, SEditableTextBox)
+		.OnTextChanged(this, &FDataprepSetOutputFolderDetails::FolderName_TextChanged)
+		.OnTextCommitted(this, &FDataprepSetOutputFolderDetails::FolderName_TextCommited)
+		.Text(FText::FromString(Operation->FolderName))
+	];
+}
+
+void FDataprepSetOutputFolderDetails::FolderName_TextCommited(const FText& InText, ETextCommit::Type InCommitType)
+{
+	if (bValidFolderName)
+	{
+		FolderNamePropertyHandle->SetValue(InText.ToString());
+	}
+	else
+	{
+		// New name is not valid: revert to old folder name
+		TextBox->SetText(FText::FromString(Operation->FolderName));
+	}
+
+	bValidFolderName = true;
+}
+
+void FDataprepSetOutputFolderDetails::FolderName_TextChanged(const FText& Text)
+{
+	// Slash and Square brackets are invalid characters for a folder name
+	const FString InvalidChars = INVALID_LONGPACKAGE_CHARACTERS TEXT("/[]"); 
+
+	FText ErrorMessage;
+	FString FolderName = Text.ToString();
+
+	// See if the name contains invalid characters.
+	FString Char;
+	for( int32 CharIdx = 0; CharIdx < FolderName.Len(); ++CharIdx )
+	{
+		Char = FolderName.Mid(CharIdx, 1);
+
+		if ( InvalidChars.Contains(*Char) )
+		{
+			FString ReadableInvalidChars = InvalidChars;
+			ReadableInvalidChars.ReplaceInline(TEXT("\r"), TEXT(""));
+			ReadableInvalidChars.ReplaceInline(TEXT("\n"), TEXT(""));
+			ReadableInvalidChars.ReplaceInline(TEXT("\t"), TEXT(""));
+
+			ErrorMessage = FText::Format(LOCTEXT("InvalidFolderName_InvalidCharacters", "A folder name may not contain any of the following characters: {0}"), FText::FromString(ReadableInvalidChars));
+			break;
+		}
+	}
+
+	if (!ErrorMessage.IsEmpty() || !FFileHelper::IsFilenameValidForSaving(FolderName, ErrorMessage))
+	{
+		TextBox->SetError(ErrorMessage);
+	}
+	else
+	{
+		// Clear error
+		TextBox->SetError(FText::GetEmpty());
+	}
+
+	bValidFolderName = ErrorMessage.IsEmpty();
+}
 
 #undef LOCTEXT_NAMESPACE
diff --git a/Engine/Plugins/Enterprise/DataprepEditor/Source/DataprepLibraries/Private/DataprepOperations.h b/Engine/Plugins/Enterprise/DataprepEditor/Source/DataprepLibraries/Private/DataprepOperations.h
--- a/Engine/Plugins/Enterprise/DataprepEditor/Source/DataprepLibraries/Private/DataprepOperations.h
+++ b/Engine/Plugins/Enterprise/DataprepEditor/Source/DataprepLibraries/Private/DataprepOperations.h
@@ -474,4 +474,27 @@ public:
 protected:
 	virtual void OnExecution_Implementation(const FDataprepContext& InContext) override;
 	//~ End UDataprepOperation Interface
-};
\ No newline at end of file
+};
+
+// Customization of the details of the "Output to Folder" operation.
+class FDataprepSetOutputFolderDetails : public IDetailCustomization
+{
+public:
+	static TSharedRef< IDetailCustomization > MakeDetails() { return MakeShared<FDataprepSetOutputFolderDetails>(); };
+
+	/** Called when details should be customized */
+	virtual void CustomizeDetails( IDetailLayoutBuilder& DetailBuilder ) override;
+
+private:
+	void FolderName_TextChanged(const FText& Text);
+	void FolderName_TextCommited(const FText& InText, ETextCommit::Type InCommitType);
+
+private:
+	UDataprepSetOutputFolder* Operation = nullptr;
+
+	TSharedPtr<IPropertyHandle> FolderNamePropertyHandle;
+
+	bool bValidFolderName = true;
+
+	TSharedPtr< SEditableTextBox > TextBox;
+};
