Fixed an sorting issue in the dataprep preview scene outliner colunm.
Corrected some typos and removed some dead code.

diff --git a/Engine/Plugins/Enterprise/DataprepEditor/Source/DataprepCore/Private/SelectionSystem/DataprepSelectionSystemUtils.cpp b/Engine/Plugins/Enterprise/DataprepEditor/Source/DataprepCore/Private/SelectionSystem/DataprepSelectionSystemUtils.cpp
deleted file mode 100644
diff --git a/Engine/Plugins/Enterprise/DataprepEditor/Source/DataprepCore/Private/SelectionSystem/DataprepSelectionSystemUtils.h b/Engine/Plugins/Enterprise/DataprepEditor/Source/DataprepCore/Private/SelectionSystem/DataprepSelectionSystemUtils.h
--- a/Engine/Plugins/Enterprise/DataprepEditor/Source/DataprepCore/Private/SelectionSystem/DataprepSelectionSystemUtils.h
+++ b/Engine/Plugins/Enterprise/DataprepEditor/Source/DataprepCore/Private/SelectionSystem/DataprepSelectionSystemUtils.h
@@ -15,11 +15,6 @@
 
 namespace DataprepSelectionSystemUtils
 {
-	// Overload the function LexToString to support LexToString(const TArray<FString>&)
-	using ::LexToString;
-	FString LexToString(const TArray<FString>& StringArray);
-
-
 	template< class FilterClass, class FetcherClass, class FetchedDataType>
 	void DoFiltering(const FilterClass& Filter, const FetcherClass& Fetcher, const TArrayView<UObject*>& Objects, TFunctionRef<void (int32 Index, bool bFetchSucceded, const FetchedDataType& FetchedData)> FilteringFunction)
 	{
diff --git a/Engine/Plugins/Enterprise/DataprepEditor/Source/DataprepEditor/Private/PreviewSystem/DataprepPreviewAssetColumn.cpp b/Engine/Plugins/Enterprise/DataprepEditor/Source/DataprepEditor/Private/PreviewSystem/DataprepPreviewAssetColumn.cpp
--- a/Engine/Plugins/Enterprise/DataprepEditor/Source/DataprepEditor/Private/PreviewSystem/DataprepPreviewAssetColumn.cpp
+++ b/Engine/Plugins/Enterprise/DataprepEditor/Source/DataprepEditor/Private/PreviewSystem/DataprepPreviewAssetColumn.cpp
@@ -20,7 +20,7 @@ FDataprepPreviewAssetColumn::FDataprepPreviewAssetColumn(const TSharedRef<FDatap
 SHeaderRow::FColumn::FArguments FDataprepPreviewAssetColumn::ConstructHeaderRowColumn(const TSharedRef<AssetPreviewWidget::SAssetsPreviewWidget>& PreviewWidget)
 {
 	PreviewWidgetWeakPtr = PreviewWidget;
-	PreviewSystem->GetOnPreviewIsDoneProcessing().AddSP( this, &FDataprepPreviewAssetColumn::OnPreviewSytemIsDoneProcessing );
+	PreviewSystem->GetOnPreviewIsDoneProcessing().AddSP( this, &FDataprepPreviewAssetColumn::OnPreviewSystemIsDoneProcessing );
 
 	return SHeaderRow::Column( GetColumnID() )
 		.DefaultLabel( LOCTEXT("Preview_HeaderText", "Preview") )
@@ -70,14 +70,14 @@ void FDataprepPreviewAssetColumn::SortItems(TArray<AssetPreviewWidget::IAssetTre
 		{
 			if ( FirstPreviewData->Status == SecondPreviewData->Status && FirstPreviewData->CurrentProcessingIndex == SecondPreviewData->CurrentProcessingIndex )
 			{
-				EDataprepPreviewResultComparaison Comparaison = FirstPreviewData->CompareFetchedDataTo( *SecondPreviewData );
-				if ( Comparaison != EDataprepPreviewResultComparaison::Equal )
+				EDataprepPreviewResultComparison Comparaison = FirstPreviewData->CompareFetchedDataTo( *SecondPreviewData );
+				if ( Comparaison != EDataprepPreviewResultComparison::Equal )
 				{
 					if ( SortMode == EColumnSortMode::Descending )
 					{
-						return Comparaison == EDataprepPreviewResultComparaison::BiggerThan;
+						return Comparaison == EDataprepPreviewResultComparison::BiggerThan;
 					}
-					return Comparaison == EDataprepPreviewResultComparaison::SmallerThan;
+					return Comparaison == EDataprepPreviewResultComparison::SmallerThan;
 				}
 			}
 			else if ( FirstPreviewData->Status == EDataprepPreviewStatus::Pass )
@@ -105,7 +105,7 @@ void FDataprepPreviewAssetColumn::SortItems(TArray<AssetPreviewWidget::IAssetTre
 }
 
 
-void FDataprepPreviewAssetColumn::OnPreviewSytemIsDoneProcessing()
+void FDataprepPreviewAssetColumn::OnPreviewSystemIsDoneProcessing()
 {
 	if (TSharedPtr<AssetPreviewWidget::SAssetsPreviewWidget> PreviewWidget = PreviewWidgetWeakPtr.Pin())
 	{
diff --git a/Engine/Plugins/Enterprise/DataprepEditor/Source/DataprepEditor/Private/PreviewSystem/DataprepPreviewAssetColumn.h b/Engine/Plugins/Enterprise/DataprepEditor/Source/DataprepEditor/Private/PreviewSystem/DataprepPreviewAssetColumn.h
--- a/Engine/Plugins/Enterprise/DataprepEditor/Source/DataprepEditor/Private/PreviewSystem/DataprepPreviewAssetColumn.h
+++ b/Engine/Plugins/Enterprise/DataprepEditor/Source/DataprepEditor/Private/PreviewSystem/DataprepPreviewAssetColumn.h
@@ -38,7 +38,7 @@ public:
 
 private:
 	// Request a sort refresh if this column was the one doing the sorting
-	void OnPreviewSytemIsDoneProcessing();
+	void OnPreviewSystemIsDoneProcessing();
 
 	TWeakPtr<AssetPreviewWidget::SAssetsPreviewWidget> PreviewWidgetWeakPtr;
 	TSharedPtr<FDataprepPreviewSystem> PreviewSystem;
diff --git a/Engine/Plugins/Enterprise/DataprepEditor/Source/DataprepEditor/Private/PreviewSystem/DataprepPreviewSceneOutlinerColumn.cpp b/Engine/Plugins/Enterprise/DataprepEditor/Source/DataprepEditor/Private/PreviewSystem/DataprepPreviewSceneOutlinerColumn.cpp
--- a/Engine/Plugins/Enterprise/DataprepEditor/Source/DataprepEditor/Private/PreviewSystem/DataprepPreviewSceneOutlinerColumn.cpp
+++ b/Engine/Plugins/Enterprise/DataprepEditor/Source/DataprepEditor/Private/PreviewSystem/DataprepPreviewSceneOutlinerColumn.cpp
@@ -46,7 +46,7 @@ FName FDataprepPreviewOutlinerColumn::GetColumnID()
 
 SHeaderRow::FColumn::FArguments FDataprepPreviewOutlinerColumn::ConstructHeaderRowColumn()
 {
-	CachedPreviewData->GetOnPreviewIsDoneProcessing().AddSP( this, &FDataprepPreviewOutlinerColumn::OnPreviewSytemIsDoneProcessing );
+	CachedPreviewData->GetOnPreviewIsDoneProcessing().AddSP( this, &FDataprepPreviewOutlinerColumn::OnPreviewSystemIsDoneProcessing );
 
 	return SHeaderRow::Column( GetColumnID() )
 		.DefaultLabel( LOCTEXT("Preview_HeaderText", "Preview") )
@@ -103,14 +103,14 @@ void FDataprepPreviewOutlinerColumn::SortItems(TArray<SceneOutliner::FTreeItemPt
 				{
 					if ( FirstPreviewData->Status == SecondPreviewData->Status && FirstPreviewData->CurrentProcessingIndex == SecondPreviewData->CurrentProcessingIndex )
 					{
-						EDataprepPreviewResultComparaison Comparaison = FirstPreviewData->CompareFetchedDataTo( *SecondPreviewData );
-						if ( Comparaison != EDataprepPreviewResultComparaison::Equal )
+						EDataprepPreviewResultComparison Comparaison = FirstPreviewData->CompareFetchedDataTo( *SecondPreviewData );
+						if ( Comparaison != EDataprepPreviewResultComparison::Equal )
 						{
 							if ( SortMode == EColumnSortMode::Descending )
 							{
-								return Comparaison == EDataprepPreviewResultComparaison::BiggerThan;
+								return Comparaison == EDataprepPreviewResultComparison::BiggerThan;
 							}
-							return Comparaison == EDataprepPreviewResultComparaison::SmallerThan;
+							return Comparaison == EDataprepPreviewResultComparison::SmallerThan;
 						}
 					}
 					else if ( FirstPreviewData->Status == EDataprepPreviewStatus::Pass )
@@ -135,7 +135,7 @@ void FDataprepPreviewOutlinerColumn::SortItems(TArray<SceneOutliner::FTreeItemPt
 		// If all else fail filter by name (always Ascending)
 		int32 SortPriorityFirst = First->GetTypeSortPriority();
 		int32 SortPrioritySecond = Second->GetTypeSortPriority();
-		if ( SortPriorityFirst != SortPriorityFirst )
+		if ( SortPriorityFirst != SortPrioritySecond )
 		{
 			return SortPriorityFirst < SortPrioritySecond;
 		}
@@ -144,7 +144,7 @@ void FDataprepPreviewOutlinerColumn::SortItems(TArray<SceneOutliner::FTreeItemPt
 	});
 }
 
-void FDataprepPreviewOutlinerColumn::OnPreviewSytemIsDoneProcessing()
+void FDataprepPreviewOutlinerColumn::OnPreviewSystemIsDoneProcessing()
 {
 	if ( TSharedPtr<ISceneOutliner> SceneOutliner = WeakSceneOutliner.Pin() )
 	{
diff --git a/Engine/Plugins/Enterprise/DataprepEditor/Source/DataprepEditor/Private/PreviewSystem/DataprepPreviewSceneOutlinerColumn.h b/Engine/Plugins/Enterprise/DataprepEditor/Source/DataprepEditor/Private/PreviewSystem/DataprepPreviewSceneOutlinerColumn.h
--- a/Engine/Plugins/Enterprise/DataprepEditor/Source/DataprepEditor/Private/PreviewSystem/DataprepPreviewSceneOutlinerColumn.h
+++ b/Engine/Plugins/Enterprise/DataprepEditor/Source/DataprepEditor/Private/PreviewSystem/DataprepPreviewSceneOutlinerColumn.h
@@ -32,7 +32,7 @@ public:
 
 private:
 	// Request a sort refresh if this column was the one doing the sorting
-	void OnPreviewSytemIsDoneProcessing();
+	void OnPreviewSystemIsDoneProcessing();
 
 	TWeakPtr<ISceneOutliner> WeakSceneOutliner;
 
diff --git a/Engine/Plugins/Enterprise/DataprepEditor/Source/DataprepEditor/Private/PreviewSystem/DataprepPreviewSystem.cpp b/Engine/Plugins/Enterprise/DataprepEditor/Source/DataprepEditor/Private/PreviewSystem/DataprepPreviewSystem.cpp
--- a/Engine/Plugins/Enterprise/DataprepEditor/Source/DataprepEditor/Private/PreviewSystem/DataprepPreviewSystem.cpp
+++ b/Engine/Plugins/Enterprise/DataprepEditor/Source/DataprepEditor/Private/PreviewSystem/DataprepPreviewSystem.cpp
@@ -19,50 +19,50 @@
 
 namespace DataprepPreviewUtils
 {
-	struct FPreviewVariantComarator
+	struct FPreviewVariantComparator
 	{
-		FPreviewVariantComarator(FDataprepPreviewProcessingResult& InCallingObject)
+		FPreviewVariantComparator(FDataprepPreviewProcessingResult& InCallingObject)
 			: Caller( InCallingObject )
 		{}
 
-		EDataprepPreviewResultComparaison operator()(const FEmptyVariantState&)
+		EDataprepPreviewResultComparison operator()(const FEmptyVariantState&)
 		{
-			return EDataprepPreviewResultComparaison::Equal;
+			return EDataprepPreviewResultComparison::Equal;
 		}
 
-		EDataprepPreviewResultComparaison operator()(int32 Value)
+		EDataprepPreviewResultComparison operator()(int32 Value)
 		{
 			int32 CallerValue = Caller.FetchedData.Get<int32>();
 			if ( CallerValue == Value )
 			{
-				return EDataprepPreviewResultComparaison::Equal;
+				return EDataprepPreviewResultComparison::Equal;
 			}
-			return CallerValue > Value ? EDataprepPreviewResultComparaison::BiggerThan : EDataprepPreviewResultComparaison::SmallerThan;
+			return CallerValue > Value ? EDataprepPreviewResultComparison::BiggerThan : EDataprepPreviewResultComparison::SmallerThan;
 		}
 
-		EDataprepPreviewResultComparaison operator()(float Value)
+		EDataprepPreviewResultComparison operator()(float Value)
 		{
 			float CallerValue = Caller.FetchedData.Get<float>();
 			if ( CallerValue == Value )
 			{
-				return EDataprepPreviewResultComparaison::Equal;
+				return EDataprepPreviewResultComparison::Equal;
 			}
-			return CallerValue > Value ? EDataprepPreviewResultComparaison::BiggerThan : EDataprepPreviewResultComparaison::SmallerThan;
+			return CallerValue > Value ? EDataprepPreviewResultComparison::BiggerThan : EDataprepPreviewResultComparison::SmallerThan;
 		}
 
-		EDataprepPreviewResultComparaison operator()(const FString& Value)
+		EDataprepPreviewResultComparison operator()(const FString& Value)
 		{
 			const FString& CallerValue = Caller.FetchedData.Get<FString>();
 			int32 CompareResult = CallerValue.Compare( Value );
 			if ( CompareResult == 0 )
 			{
-				return EDataprepPreviewResultComparaison::Equal;
+				return EDataprepPreviewResultComparison::Equal;
 			}
-			return CompareResult > 0 ? EDataprepPreviewResultComparaison::BiggerThan : EDataprepPreviewResultComparaison::SmallerThan;
+			return CompareResult > 0 ? EDataprepPreviewResultComparison::BiggerThan : EDataprepPreviewResultComparison::SmallerThan;
 		}
 
 		template<class T>
-		EDataprepPreviewResultComparaison operator()(T) = delete;
+		EDataprepPreviewResultComparison operator()(T) = delete;
 
 		FDataprepPreviewProcessingResult& Caller;
 	};
@@ -137,14 +137,14 @@ namespace DataprepPreviewUtils
 
 const int32 FDataprepPreviewSystem::IncrementalCount = 2000;
 
-EDataprepPreviewResultComparaison FDataprepPreviewProcessingResult::CompareFetchedDataTo(const FDataprepPreviewProcessingResult& Other)
+EDataprepPreviewResultComparison FDataprepPreviewProcessingResult::CompareFetchedDataTo(const FDataprepPreviewProcessingResult& Other)
 {
 	if ( FetchedData.GetIndex() != Other.FetchedData.GetIndex() )
 	{
-		return EDataprepPreviewResultComparaison::Equal;
+		return EDataprepPreviewResultComparison::Equal;
 	}
 
-	return Visit( DataprepPreviewUtils::FPreviewVariantComarator( *this ), Other.FetchedData );
+	return Visit( DataprepPreviewUtils::FPreviewVariantComparator( *this ), Other.FetchedData );
 }
 
 FText FDataprepPreviewProcessingResult::GetFetchedDataAsText() const
diff --git a/Engine/Plugins/Enterprise/DataprepEditor/Source/DataprepEditor/Private/PreviewSystem/DataprepPreviewSystem.h b/Engine/Plugins/Enterprise/DataprepEditor/Source/DataprepEditor/Private/PreviewSystem/DataprepPreviewSystem.h
--- a/Engine/Plugins/Enterprise/DataprepEditor/Source/DataprepEditor/Private/PreviewSystem/DataprepPreviewSystem.h
+++ b/Engine/Plugins/Enterprise/DataprepEditor/Source/DataprepEditor/Private/PreviewSystem/DataprepPreviewSystem.h
@@ -26,7 +26,7 @@ enum class EDataprepPreviewStatus : uint8
 	// There is assumption on the order of these so that the expression Pass > Failed return true
 };
 
-enum class EDataprepPreviewResultComparaison : uint8
+enum class EDataprepPreviewResultComparison : uint8
 {
 	SmallerThan,
 	Equal,
@@ -41,7 +41,7 @@ struct FDataprepPreviewProcessingResult
 		, FetchedData()
 	{}
 
-	EDataprepPreviewResultComparaison CompareFetchedDataTo(const FDataprepPreviewProcessingResult& Other);
+	EDataprepPreviewResultComparison CompareFetchedDataTo(const FDataprepPreviewProcessingResult& Other);
 
 	FText GetFetchedDataAsText() const;
 
diff --git a/Engine/Plugins/Enterprise/DataprepEditor/Source/DataprepEditor/Private/Widgets/DataprepGraph/SDataprepGraphEditor.cpp b/Engine/Plugins/Enterprise/DataprepEditor/Source/DataprepEditor/Private/Widgets/DataprepGraph/SDataprepGraphEditor.cpp
--- a/Engine/Plugins/Enterprise/DataprepEditor/Source/DataprepEditor/Private/Widgets/DataprepGraph/SDataprepGraphEditor.cpp
+++ b/Engine/Plugins/Enterprise/DataprepEditor/Source/DataprepEditor/Private/Widgets/DataprepGraph/SDataprepGraphEditor.cpp
@@ -877,7 +877,7 @@ FActionMenuContent SDataprepGraphEditor::OnCreateNodeOrPinMenu(UEdGraph* Current
 							});
 
 						MenuBuilder->AddMenuEntry(LOCTEXT("SetFilterPreview", "Preview Filter(s)"),
-							bAreFilterFromSameAction ? LOCTEXT("SetFilterPreviewTooltip", "Change the columns of the scene preview and asset preview tabs to display a preview of what the filters would select from the current scene") : LOCTEXT("SetFilterPreviewFailTooltip", "The filters can only be previewd if they are from the same action"),
+							bAreFilterFromSameAction ? LOCTEXT("SetFilterPreviewTooltip", "Change the columns of the scene preview and asset preview tabs to display a preview of what the filters would select from the current scene") : LOCTEXT("SetFilterPreviewFailTooltip", "The filters can only be previewed if they are from the same action."),
 							FSlateIcon(),
 							SetFilterPreview);
 					}
