Fixed annoying movement when entering and leaving the graph panel on a drag and drop

diff --git a/Engine/Plugins/Enterprise/DataprepEditor/Source/DataprepEditor/Private/SchemaActions/DataprepDragDropOp.cpp b/Engine/Plugins/Enterprise/DataprepEditor/Source/DataprepEditor/Private/SchemaActions/DataprepDragDropOp.cpp
--- a/Engine/Plugins/Enterprise/DataprepEditor/Source/DataprepEditor/Private/SchemaActions/DataprepDragDropOp.cpp
+++ b/Engine/Plugins/Enterprise/DataprepEditor/Source/DataprepEditor/Private/SchemaActions/DataprepDragDropOp.cpp
@@ -7,17 +7,19 @@
 #include "DataprepAsset.h"
 #include "DataprepCoreUtils.h"
 #include "DataprepEditorLogCategory.h"
+#include "DataprepEditorStyle.h"
 #include "DataprepGraph/DataprepGraph.h"
 #include "DataprepGraph/DataprepGraphActionNode.h"
 #include "DataprepSchemaActionUtils.h"
+#include "Widgets/DataprepGraph/SDataprepGraphActionNode.h"
 #include "Widgets/DataprepGraph/SDataprepGraphActionStepNode.h"
+#include "Widgets/DataprepGraph/SDataprepGraphTrackNode.h"
 
 #include "EdGraph/EdGraph.h"
 #include "EdGraphSchema_K2.h"
 #include "EditorStyleSet.h"
 #include "Kismet2/BlueprintEditorUtils.h"
 #include "ScopedTransaction.h"
-#include "SGraphPanel.h"
 #include "SPinTypeSelector.h"
 #include "Templates/UnrealTypeTraits.h"
 #include "UObject/StrongObjectPtr.h"
@@ -29,6 +31,7 @@
 #include "Widgets/Layout/SConstraintCanvas.h"
 #include "Widgets/Layout/SScaleBox.h"
 #include "Widgets/SBoxPanel.h"
+#include "Widgets/SOverlay.h"
 #include "Widgets/Text/STextBlock.h"
 
 #define LOCTEXT_NAMESPACE "DataprepDragAndDrop"
@@ -56,11 +59,11 @@ TSharedRef<FDataprepDragDropOp> FDataprepDragDropOp::New(FDataprepGraphOperation
 	return DragDrop;
 }
 
-TSharedRef<FDataprepDragDropOp> FDataprepDragDropOp::New(const TSharedRef<SGraphPanel>& InGraphPanel, const TSharedRef<SDataprepGraphActionStepNode>& InDraggedNode)
+TSharedRef<FDataprepDragDropOp> FDataprepDragDropOp::New(const TSharedRef<SDataprepGraphTrackNode>& InTrackNode, const TSharedRef<SDataprepGraphActionStepNode>& InDraggedNode)
 {
 	TSharedRef<FDataprepDragDropOp> Operation = MakeShareable(new FDataprepDragDropOp);
 
-	Operation->GraphPanelPtr = InGraphPanel;
+	Operation->TrackNodePtr = InTrackNode;
 	Operation->DraggedNodeWidgets.Add(InDraggedNode);
 	if(UDataprepGraphActionStepNode* ActionStepNode = Cast<UDataprepGraphActionStepNode>(InDraggedNode->GetNodeObj()))
 	{
@@ -92,6 +95,12 @@ TSharedRef<FDataprepDragDropOp> FDataprepDragDropOp::New(UDataprepActionStep* In
 	return Operation;
 }
 
+void FDataprepDragDropOp::Construct()
+{
+	FGraphEditorDragDropAction::Construct();
+	CursorDecoratorWindow->SetOpacity(0.9f);
+}
+
 void FDataprepDragDropOp::HoverTargetChanged()
 {
 	FText DrapDropText;
@@ -110,17 +119,17 @@ void FDataprepDragDropOp::HoverTargetChanged()
 	else if(Cast<UDataprepGraphActionStepNode>(GetHoveredNode()) != nullptr)
 	{
 		bDropTargetValid = true;
-		DrapDropText = LOCTEXT("TargetIsDataprepActionNode", "Add/Insert step to location");
+		DrapDropText = LOCTEXT("CopyDataprepActionStepNode", "Add/Insert step to location");
 	}
 	else if(Cast<UDataprepGraphActionNode>(GetHoveredNode()) != nullptr)
 	{
 		bDropTargetValid = true;
-		DrapDropText = LOCTEXT("TargetIsDataprepActionNode", "Add/Insert step to location");
+		DrapDropText = LOCTEXT("CopyDataprepActionAssetNode", "Add/Insert step to location");
 	}
 	else if(Cast<UDataprepGraphRecipeNode>(GetHoveredNode()) != nullptr)
 	{
 		bDropTargetValid = true;
-		DrapDropText = LOCTEXT("TargetIsDataprepRecipeNode", "Insert action to location");
+		DrapDropText = LOCTEXT("InsertDataprepActionAssetNode", "Insert action to location");
 	}
 	else if ( UEdGraph* EdGraph = GetHoveredGraph() )
 	{
@@ -182,9 +191,9 @@ void FDataprepDragDropOp::OnDragged(const FDragDropEvent& DragDropEvent)
 	CursorDecoratorWindow->MoveWindowTo(TargetPosition + DecoratorAdjust);
 	// Request the active panel to scroll if required
 
-	if(SGraphPanel* GraphPanel = GraphPanelPtr.Get())
+	if(SDataprepGraphTrackNode* TrackNode = TrackNodePtr.Pin().Get())
 	{
-		GraphPanel->RequestDeferredPan(TargetPosition);
+		TrackNode->RequestViewportPan(TargetPosition);
 	}
 
 	Super::OnDragged(DragDropEvent);
@@ -403,7 +412,7 @@ FReply FDataprepDragDropOp::DoDropOnActionStep(UDataprepGraphActionStepNode* Tar
 
 	if(!bDropTargetValid && !bCopyRequested)
 	{
-		return FReply::Handled().EndDragDrop();
+		return FReply::Unhandled().EndDragDrop();
 	}
 
 	UDataprepActionAsset* TargetActionAsset = TargetActionStepNode->GetDataprepActionAsset();
@@ -491,7 +500,10 @@ void FDataprepDragDropOp::DropStepFromPanel(UDataprepActionAsset* TargetActionAs
 
 FReply FDataprepDragDropOp::DoDropOnTrack(UDataprepAsset* TargetDataprepAsset, int32 InsertIndex)
 {
-	ensure(Cast<UDataprepGraphRecipeNode>(GetHoveredNode()));
+	if(Cast<UDataprepGraphRecipeNode>(GetHoveredNode()) == nullptr)
+	{
+		return FReply::Handled().EndDragDrop();
+	}
 
 	if(DraggedSteps.Num() > 0)
 	{
diff --git a/Engine/Plugins/Enterprise/DataprepEditor/Source/DataprepEditor/Private/SchemaActions/DataprepDragDropOp.h b/Engine/Plugins/Enterprise/DataprepEditor/Source/DataprepEditor/Private/SchemaActions/DataprepDragDropOp.h
--- a/Engine/Plugins/Enterprise/DataprepEditor/Source/DataprepEditor/Private/SchemaActions/DataprepDragDropOp.h
+++ b/Engine/Plugins/Enterprise/DataprepEditor/Source/DataprepEditor/Private/SchemaActions/DataprepDragDropOp.h
@@ -11,7 +11,7 @@
 
 struct FDataprepSchemaActionContext;
 class SDataprepGraphActionStepNode;
-class SGraphPanel;
+class SDataprepGraphTrackNode;
 class UDataprepActionAsset;
 class UDataprepAsset;
 class UDataprepGraphActionNode;
@@ -36,7 +36,7 @@ public:
 
 	static TSharedRef<FDataprepDragDropOp> New(TSharedRef<FDataprepSchemaAction> InAction);
 	static TSharedRef<FDataprepDragDropOp> New(FDataprepGraphOperation&& GraphOperation);
-	static TSharedRef<FDataprepDragDropOp> New(const TSharedRef<SGraphPanel>& InGraphPanel, const TSharedRef<SDataprepGraphActionStepNode>& InDraggedNode);
+	static TSharedRef<FDataprepDragDropOp> New(const TSharedRef<SDataprepGraphTrackNode>& InTrackNode, const TSharedRef<SDataprepGraphActionStepNode>& InDraggedNode);
 	static TSharedRef<FDataprepDragDropOp> New(UDataprepActionStep* InActionStep);
 
 	void SetHoveredDataprepActionContext(TOptional<FDataprepSchemaActionContext> Context);
@@ -50,6 +50,7 @@ public:
 	virtual void OnDragged (const class FDragDropEvent& DragDropEvent ) override;
 	virtual EVisibility GetIconVisible() const;
 	virtual EVisibility GetErrorIconVisible() const;
+	virtual void Construct() override;
 	// End of FGraphEditorDragDropAction Interface
 
 	/**
@@ -66,9 +67,9 @@ public:
 	// Returns action step node targeted for the drop
 	UDataprepGraphActionStepNode* GetDropTargetNode() const;
 
-	void SetGraphPanel(const TSharedPtr<SGraphPanel>& InGraphPanel)
+	void SetTrackNode(const TSharedPtr<SDataprepGraphTrackNode>& InTrackNodePtr)
 	{
-		GraphPanelPtr = InGraphPanel;
+		TrackNodePtr = InTrackNodePtr;
 	}
 
 	bool IsValidDrop() { return bDropTargetValid; }
@@ -109,7 +110,7 @@ private:
 	typedef TTuple<TWeakObjectPtr<UDataprepActionAsset>,int32,TWeakObjectPtr<UDataprepActionStep>> FDraggedStepEntry;
 
 	/** Graph panel associated with the Dataprep graph editor */
-	TSharedPtr<SGraphPanel> GraphPanelPtr;
+	TWeakPtr<SDataprepGraphTrackNode> TrackNodePtr;
 
 	/** Array of action steps being dragged */
 	TArray<TSharedRef<SDataprepGraphActionStepNode>> DraggedNodeWidgets;
diff --git a/Engine/Plugins/Enterprise/DataprepEditor/Source/DataprepEditor/Private/Widgets/DataprepGraph/SDataprepGraphActionNode.h b/Engine/Plugins/Enterprise/DataprepEditor/Source/DataprepEditor/Private/Widgets/DataprepGraph/SDataprepGraphActionNode.h
--- a/Engine/Plugins/Enterprise/DataprepEditor/Source/DataprepEditor/Private/Widgets/DataprepGraph/SDataprepGraphActionNode.h
+++ b/Engine/Plugins/Enterprise/DataprepEditor/Source/DataprepEditor/Private/Widgets/DataprepGraph/SDataprepGraphActionNode.h
@@ -47,6 +47,7 @@ public:
 	// End of SGraphNode interface
 
 	void SetParentTrackNode(TSharedPtr<SDataprepGraphTrackNode> InParentTrackNode);
+	TSharedPtr<SDataprepGraphTrackNode> GetParentTrackNode() { return ParentTrackNodePtr.Pin(); }
 
 	int32 GetExecutionOrder() const { return ExecutionOrder; }
 	void UpdateExecutionOrder();
diff --git a/Engine/Plugins/Enterprise/DataprepEditor/Source/DataprepEditor/Private/Widgets/DataprepGraph/SDataprepGraphActionStepNode.cpp b/Engine/Plugins/Enterprise/DataprepEditor/Source/DataprepEditor/Private/Widgets/DataprepGraph/SDataprepGraphActionStepNode.cpp
--- a/Engine/Plugins/Enterprise/DataprepEditor/Source/DataprepEditor/Private/Widgets/DataprepGraph/SDataprepGraphActionStepNode.cpp
+++ b/Engine/Plugins/Enterprise/DataprepEditor/Source/DataprepEditor/Private/Widgets/DataprepGraph/SDataprepGraphActionStepNode.cpp
@@ -287,8 +287,11 @@ FReply SDataprepGraphActionStepNode::OnDragDetected(const FGeometry & MyGeometry
 	{
 		if (UDataprepActionStep* ActionStep = ActionStepNode->GetDataprepActionStep())
 		{
-			ParentNodePtr.Pin()->SetDraggedIndex(StepIndex);
-			return FReply::Handled().BeginDragDrop(FDataprepDragDropOp::New( GetOwnerPanel().ToSharedRef(), SharedThis(this)));
+			if(SDataprepGraphActionNode* ActionNode = ParentNodePtr.Pin().Get())
+			{
+				ActionNode->SetDraggedIndex(StepIndex);
+				return FReply::Handled().BeginDragDrop(FDataprepDragDropOp::New( SharedThis(ActionNode->GetParentTrackNode().Get()), SharedThis(this)));
+			}
 		}
 	}
 
diff --git a/Engine/Plugins/Enterprise/DataprepEditor/Source/DataprepEditor/Private/Widgets/DataprepGraph/SDataprepGraphEditor.cpp b/Engine/Plugins/Enterprise/DataprepEditor/Source/DataprepEditor/Private/Widgets/DataprepGraph/SDataprepGraphEditor.cpp
--- a/Engine/Plugins/Enterprise/DataprepEditor/Source/DataprepEditor/Private/Widgets/DataprepGraph/SDataprepGraphEditor.cpp
+++ b/Engine/Plugins/Enterprise/DataprepEditor/Source/DataprepEditor/Private/Widgets/DataprepGraph/SDataprepGraphEditor.cpp
@@ -349,7 +349,7 @@ void SDataprepGraphEditor::OnDragEnter(const FGeometry & MyGeometry, const FDrag
 	if (TrackGraphNodePtr.IsValid() && DragNodeOp.IsValid())
 	{
 		// Inform the Drag and Drop operation that we are hovering over this node.
-		DragNodeOp->SetGraphPanel(TrackGraphNodePtr.Pin()->GetOwnerPanel());
+		DragNodeOp->SetTrackNode(TrackGraphNodePtr.Pin());
 	}
 
 	SGraphEditor::OnDragEnter(MyGeometry, DragDropEvent);
@@ -371,8 +371,8 @@ void SDataprepGraphEditor::OnDragLeave(const FDragDropEvent & DragDropEvent)
 	TSharedPtr<FDataprepDragDropOp> DragNodeOp = DragDropEvent.GetOperationAs<FDataprepDragDropOp>();
 	if (TrackGraphNodePtr.IsValid() && DragNodeOp.IsValid())
 	{
-		// Inform the Drag and Drop operation that we are hovering over this node.
-		DragNodeOp->SetGraphPanel(TSharedPtr<SGraphPanel>());
+		// Inform the Drag and Drop operation that we are not hovering over this node anymore.
+		DragNodeOp->SetTrackNode(TSharedPtr<SDataprepGraphTrackNode>());
 	}
 
 	SGraphEditor::OnDragLeave(DragDropEvent);
@@ -383,9 +383,6 @@ FReply SDataprepGraphEditor::OnDrop(const FGeometry& MyGeometry, const FDragDrop
 	TSharedPtr<FDataprepDragDropOp> DragNodeOp = DragDropEvent.GetOperationAs<FDataprepDragDropOp>();
 	if (TrackGraphNodePtr.IsValid() && DragNodeOp.IsValid())
 	{
-		//const FVector2D NodeAddPosition = (MyGeometry.AbsoluteToLocal( DragDropEvent.GetScreenSpacePosition() ) / LastZoomAmount) + LastLocation;
-		//return DragNodeOp->DroppedOnPanel(SharedThis(this), DragDropEvent.GetScreenSpacePosition(), NodeAddPosition, *GetCurrentGraph()).EndDragDrop();
-
 		return FReply::Handled().EndDragDrop();
 	}
 
diff --git a/Engine/Plugins/Enterprise/DataprepEditor/Source/DataprepEditor/Private/Widgets/DataprepGraph/SDataprepGraphTrackNode.cpp b/Engine/Plugins/Enterprise/DataprepEditor/Source/DataprepEditor/Private/Widgets/DataprepGraph/SDataprepGraphTrackNode.cpp
--- a/Engine/Plugins/Enterprise/DataprepEditor/Source/DataprepEditor/Private/Widgets/DataprepGraph/SDataprepGraphTrackNode.cpp
+++ b/Engine/Plugins/Enterprise/DataprepEditor/Source/DataprepEditor/Private/Widgets/DataprepGraph/SDataprepGraphTrackNode.cpp
@@ -541,6 +541,53 @@ void SDataprepGraphTrackNode::OnControlKeyChanged(bool bControlKeyDown)
 	}
 }
 
+void SDataprepGraphTrackNode::RequestViewportPan(const FVector2D& InScreenSpacePosition)
+{
+	// Note: Code copied from SNodePanel::RequestDeferredPan and its subsequent calls. Calling SNodePanel::RequestDeferredPan is not an option
+	//		 since it changes the view offset during the call to OnPaint which bypasses the adjustments done by the SDataprepGraphEditor in
+	//		 SDataprepGraphEditor::Tick 
+	if(SGraphPanel* GraphPanel = OwnerGraphPanelPtr.Pin().Get())
+	{
+		// How quickly to ramp up the pan speed as the user moves the mouse further past the edge of the graph panel.
+		constexpr float EdgePanSpeedCoefficient = 2.f;
+		constexpr float EdgePanSpeedPower = 0.6f;
+		// Never pan faster than this - probably not really required since we raise to a power of 0.6
+		constexpr float MaxPanSpeed = 200.0f;
+		// Start panning before we reach the edge of the graph panel.
+		constexpr float EdgePanForgivenessZone = 30.0f;
+
+		//const FVector2D PanAmount = ComputeEdgePanAmount( GraphPanel->GetTickSpaceGeometry(), InScreenSpacePosition ) / GraphPanel->GetZoomAmount();
+		const FGeometry& PanelGeometry = GraphPanel->GetTickSpaceGeometry();
+		const FVector2D PanelPosition = PanelGeometry.AbsoluteToLocal( InScreenSpacePosition );
+		const FVector2D PanelSize = PanelGeometry.GetLocalSize();
+		const float ZoomAmount = GraphPanel->GetZoomAmount();
+
+		FVector2D PanAmount(FVector2D::ZeroVector);
+		if ( PanelPosition.X <= EdgePanForgivenessZone )
+		{
+			PanAmount.X = FMath::Max( -MaxPanSpeed, EdgePanSpeedCoefficient * -FMath::Pow(EdgePanForgivenessZone - PanelPosition.X, EdgePanSpeedPower) );
+		}
+		else if( PanelPosition.X >= PanelSize.X - EdgePanForgivenessZone )
+		{
+			PanAmount.X = FMath::Min( MaxPanSpeed, EdgePanSpeedCoefficient * FMath::Pow(PanelPosition.X - PanelSize.X + EdgePanForgivenessZone, EdgePanSpeedPower) );
+		}
+
+		if ( PanelPosition.Y <= EdgePanForgivenessZone )
+		{
+			PanAmount.Y = FMath::Max( -MaxPanSpeed, EdgePanSpeedCoefficient * -FMath::Pow(EdgePanForgivenessZone - PanelPosition.Y, EdgePanSpeedPower) );
+		}
+		else if( PanelPosition.Y >= PanelSize.Y - EdgePanForgivenessZone )
+		{
+			PanAmount.Y = FMath::Min( MaxPanSpeed, EdgePanSpeedCoefficient * FMath::Pow(PanelPosition.Y - PanelSize.Y + EdgePanForgivenessZone, EdgePanSpeedPower) );
+		}
+
+		if(PanAmount != FVector2D::ZeroVector)
+		{
+			GraphPanel->RestoreViewSettings(GraphPanel->GetViewOffset() + PanAmount / ZoomAmount, ZoomAmount);
+		}
+	}
+}
+
 void SDataprepGraphTrackNode::RequestRename(const UEdGraphNode* Node)
 {
 	if(SGraphPanel* GraphPanel = OwnerGraphPanelPtr.Pin().Get())
diff --git a/Engine/Plugins/Enterprise/DataprepEditor/Source/DataprepEditor/Private/Widgets/DataprepGraph/SDataprepGraphTrackNode.h b/Engine/Plugins/Enterprise/DataprepEditor/Source/DataprepEditor/Private/Widgets/DataprepGraph/SDataprepGraphTrackNode.h
--- a/Engine/Plugins/Enterprise/DataprepEditor/Source/DataprepEditor/Private/Widgets/DataprepGraph/SDataprepGraphTrackNode.h
+++ b/Engine/Plugins/Enterprise/DataprepEditor/Source/DataprepEditor/Private/Widgets/DataprepGraph/SDataprepGraphTrackNode.h
@@ -115,6 +115,8 @@ public:
 	/** Start editing of action asset associated to input EdGraphNode */
 	void RequestRename(const UEdGraphNode* Node);
 
+	void RequestViewportPan(const FVector2D& ScreenSpacePosition);
+
 	/** Miscellaneous values used in the display */
 	// #ueent_wip: Will be moved to the Dataprep editor's style
 	static FMargin NodePadding;
