Implemented line with bullets as suggested by designer
Added black background to hide grid in graph panel
Reworked drag-and-drop visual feedback to fit line with bullets
Fixed zoom reset on action node addition
Added logic to display the right text when hovering on top of the bottom widget of an action widget

Note that mouse cursors might still be messed up

#rnx

diff --git a/Engine/Plugins/Enterprise/DataprepEditor/Source/DataprepEditor/Private/DataprepEditorStyle.cpp b/Engine/Plugins/Enterprise/DataprepEditor/Source/DataprepEditor/Private/DataprepEditorStyle.cpp
--- a/Engine/Plugins/Enterprise/DataprepEditor/Source/DataprepEditor/Private/DataprepEditorStyle.cpp
+++ b/Engine/Plugins/Enterprise/DataprepEditor/Source/DataprepEditor/Private/DataprepEditorStyle.cpp
@@ -69,13 +69,14 @@ void FDataprepEditorStyle::Initialize()
 
 	// Dataprep action UI
 	{
+		StyleSet->Set("Dataprep.Background.Black", FLinearColor(FColor(26, 26, 26)) );
 		StyleSet->Set("Dataprep.TextSeparator.Color", FLinearColor(FColor(200, 200, 200, 200)) );
 
-		const FMargin ActionOuter(1.f, 2.f, 2.f, 2.f);
-		StyleSet->Set("DataprepAction.Body.Padding", ActionOuter + FMargin(5.f, 0.f, 0.f, 0.f));
-		StyleSet->Set("DataprepAction.Outter.Regular.Padding", ActionOuter);
+		StyleSet->Set("DataprepAction.Outter.Regular.Padding", FMargin(1.f, 2.f, 4.f, 2.f));
 		StyleSet->Set("DataprepAction.Outter.Selected.Padding", FMargin(0.f));
-		StyleSet->Set("DataprepAction.Steps.Padding", ActionOuter + FMargin(10.f, 0.f, 4.f, 0.f));
+		StyleSet->Set("DataprepAction.Body.Padding", FMargin(6.f, 2.f, 2.f, 2.f));
+		StyleSet->Set("DataprepAction.Steps.Padding", FMargin(11.f, 2.f, 6.f, 2.f));
+		StyleSet->Set("DataprepAction.DragAndDrop", FLinearColor(FColor(212, 212, 59)) );
 
 		StyleSet->Set("DataprepAction.OutlineColor", FLinearColor(FColor(10, 177, 51)));
 		StyleSet->Set("DataprepAction.BackgroundColor", FLinearColor(FColor(61, 61, 61)));
diff --git a/Engine/Plugins/Enterprise/DataprepEditor/Source/DataprepEditor/Private/Widgets/DataprepGraph/SDataprepGraphActionNode.cpp b/Engine/Plugins/Enterprise/DataprepEditor/Source/DataprepEditor/Private/Widgets/DataprepGraph/SDataprepGraphActionNode.cpp
--- a/Engine/Plugins/Enterprise/DataprepEditor/Source/DataprepEditor/Private/Widgets/DataprepGraph/SDataprepGraphActionNode.cpp
+++ b/Engine/Plugins/Enterprise/DataprepEditor/Source/DataprepEditor/Private/Widgets/DataprepGraph/SDataprepGraphActionNode.cpp
@@ -26,6 +26,9 @@
 
 #define LOCTEXT_NAMESPACE "DataprepGraphEditor"
 
+float SDataprepGraphActionNode::DefaultWidth = 300.f;
+float SDataprepGraphActionNode::DefaultHeight = 100.f;
+
 class SDataprepGraphActionProxyNode : public SGraphNode
 {
 public:
@@ -85,11 +88,15 @@ public:
 	{
 		return  FEditorStyle::GetNoBrush();
 	}
+
+	virtual bool ShouldAllowCulling() const override { return false; }
 	// End of SGraphNode interface
 
 	FVector2D GetSize()
 	{
-		FVector2D Size(10.f);
+		static FVector2D DefaultSize(SDataprepGraphActionNode::DefaultWidth, SDataprepGraphActionNode::DefaultHeight);
+
+		FVector2D Size(DefaultSize);
 
 		if(SDataprepGraphActionNode* ParentNode = ParentNodePtr.Pin().Get())
 		{
@@ -100,7 +107,7 @@ public:
 				Size = ParentNode->GetDesiredSize();
 				if(Size == FVector2D::ZeroVector)
 				{
-					Size.Set(10.f, 10.f);
+					Size = DefaultSize;
 				}
 			}
 		}
@@ -537,16 +544,6 @@ TSharedRef<SWidget> SDataprepGraphActionNode::CreateBackground(const TAttribute<
 {
 	return SNew(SOverlay)
 
-		+SOverlay::Slot()
-		.VAlign(VAlign_Fill)
-		.HAlign(HAlign_Fill)
-		.Padding(1.f)
-		[
-			SNew(SImage)
-			.ColorAndOpacity(FLinearColor::White)
-			.Image(FDataprepEditorStyle::GetBrush( "DataprepEditor.Node.Body" ))
-		]
-
 		+SOverlay::Slot()
 		.VAlign(VAlign_Fill)
 		.HAlign(HAlign_Fill)
@@ -582,30 +579,6 @@ const FSlateBrush* SDataprepGraphActionNode::GetShadowBrush(bool bSelected) cons
 	return  FEditorStyle::GetNoBrush();
 }
 
-int32 SDataprepGraphActionNode::OnPaint(const FPaintArgs& Args, const FGeometry& AllottedGeometry, const FSlateRect& MyCullingRect, FSlateWindowElementList& OutDrawElements, int32 LayerId, const FWidgetStyle& InWidgetStyle, bool bParentEnabled) const
-{
-	// Since only a proxy is in the graph panel, draw selection outline if applicable 
-	//if(SGraphPanel* GraphPanel = GetOwnerPanel().Get())
-	//{
-	//	if (GraphPanel->SelectionManager.SelectedNodes.Contains(GraphNode))
-	//	{
-	//		const FSlateBrush* ShadowBrush = FEditorStyle::GetBrush(TEXT("Graph.Node.ShadowSelected"));
-	//		const FVector2D NodeShadowSize = GetDefault<UGraphEditorSettings>()->GetShadowDeltaSize();
-
-	//		FSlateDrawElement::MakeBox(
-	//			OutDrawElements,
-	//			LayerId,
-	//			GetPaintSpaceGeometry().ToInflatedPaintGeometry(NodeShadowSize),
-	//			ShadowBrush,
-	//			ESlateDrawEffect::None,
-	//			FLinearColor(1.0f, 1.0f, 1.0f, 1.f)
-	//		);
-	//	}
-	//}
-
-	return SGraphNode::OnPaint(Args, AllottedGeometry, MyCullingRect, OutDrawElements, LayerId, InWidgetStyle, bParentEnabled);
-}
-
 FReply SDataprepGraphActionNode::OnMouseButtonDown(const FGeometry& MyGeometry, const FPointerEvent& MouseEvent)
 {
 	BorderBackgroundColor.Set(FDataprepEditorStyle::GetColor("DataprepActionStep.DragAndDrop"));
@@ -843,6 +816,7 @@ void SDataprepGraphActionNode::PopulateActionStepListWidget()
 	.Padding(0.f, 0.f, 0.f, 10.f)
 	[
 		SAssignNew(BottomSlot, SDataprepEmptyActionStepNode, SharedThis(this))
+		.Text(TAttribute<FText>::Create(TAttribute<FText>::FGetter::CreateSP( this, &SDataprepGraphActionNode::GetBottomWidgetText ) ))
 	];
 
 	if(TrackNodePtr.IsValid())
@@ -860,4 +834,22 @@ void SDataprepGraphActionNode::OnStepsChanged()
 	}
 }
 
+FText SDataprepGraphActionNode::GetBottomWidgetText() const
+{
+	if(InsertIndex == DataprepActionPtr->GetStepsCount())
+	{
+		FModifierKeysState ModifierKeyState = FSlateApplication::Get().GetModifierKeys();
+		if(ModifierKeyState.IsControlDown() || ModifierKeyState.IsCommandDown() || DraggedIndex == INDEX_NONE)
+		{
+			return LOCTEXT("DataprepEmptyActionStepCopyLabel", "Add Step");
+		}
+		else
+		{
+			return LOCTEXT("DataprepEmptyActionStepMoveLabel", "Move Step");
+		}
+	}
+
+	return LOCTEXT("DataprepEmptyActionStepNoLabel", ".....");
+}
+
 #undef LOCTEXT_NAMESPACE
\ No newline at end of file
diff --git a/Engine/Plugins/Enterprise/DataprepEditor/Source/DataprepEditor/Private/Widgets/DataprepGraph/SDataprepGraphActionNode.h b/Engine/Plugins/Enterprise/DataprepEditor/Source/DataprepEditor/Private/Widgets/DataprepGraph/SDataprepGraphActionNode.h
--- a/Engine/Plugins/Enterprise/DataprepEditor/Source/DataprepEditor/Private/Widgets/DataprepGraph/SDataprepGraphActionNode.h
+++ b/Engine/Plugins/Enterprise/DataprepEditor/Source/DataprepEditor/Private/Widgets/DataprepGraph/SDataprepGraphActionNode.h
@@ -36,7 +36,6 @@ public:
 	virtual FCursorReply OnCursorQuery(const FGeometry& MyGeometry, const FPointerEvent& CursorEvent) const override;
 	virtual void SetOwner(const TSharedRef<SGraphPanel>& OwnerPanel) override;
 	virtual void OnDragEnter(const FGeometry& MyGeometry, const FDragDropEvent& DragDropEvent) override;
-	int32 OnPaint( const FPaintArgs& Args, const FGeometry& AllottedGeometry, const FSlateRect& MyCullingRect, FSlateWindowElementList& OutDrawElements, int32 LayerId, const FWidgetStyle& InWidgetStyle, bool bParentEnabled ) const override;
 	// End of SWidget interface
 
 	// SGraphNode interface
@@ -78,6 +77,8 @@ private:
 
 	FMargin GetOuterPadding() const;
 
+	FText GetBottomWidgetText() const;
+
 private:
 	/** Weak pointer to the associated action asset */
 	TWeakObjectPtr<class UDataprepActionAsset> DataprepActionPtr;
@@ -110,4 +111,8 @@ private:
 	TWeakPtr<FDataprepEditor> DataprepEditor;
 
 	friend SDataprepGraphActionProxyNode;
+
+public:
+	static float DefaultWidth;
+	static float DefaultHeight;
 };
diff --git a/Engine/Plugins/Enterprise/DataprepEditor/Source/DataprepEditor/Private/Widgets/DataprepGraph/SDataprepGraphActionStepNode.cpp b/Engine/Plugins/Enterprise/DataprepEditor/Source/DataprepEditor/Private/Widgets/DataprepGraph/SDataprepGraphActionStepNode.cpp
--- a/Engine/Plugins/Enterprise/DataprepEditor/Source/DataprepEditor/Private/Widgets/DataprepGraph/SDataprepGraphActionStepNode.cpp
+++ b/Engine/Plugins/Enterprise/DataprepEditor/Source/DataprepEditor/Private/Widgets/DataprepGraph/SDataprepGraphActionStepNode.cpp
@@ -45,7 +45,7 @@ void SDataprepGraphActionStepNode::Construct(const FArguments& InArgs, UDataprep
 	GraphNode = InActionStepNode;
 	DataprepEditor = InArgs._DataprepEditor;
 
-	SetCursor(EMouseCursor::ResizeUpDown);
+	SetCursor(EMouseCursor::CardinalCross);
 	UpdateGraphNode();
 }
 
diff --git a/Engine/Plugins/Enterprise/DataprepEditor/Source/DataprepEditor/Private/Widgets/DataprepGraph/SDataprepGraphEditor.cpp b/Engine/Plugins/Enterprise/DataprepEditor/Source/DataprepEditor/Private/Widgets/DataprepGraph/SDataprepGraphEditor.cpp
--- a/Engine/Plugins/Enterprise/DataprepEditor/Source/DataprepEditor/Private/Widgets/DataprepGraph/SDataprepGraphEditor.cpp
+++ b/Engine/Plugins/Enterprise/DataprepEditor/Source/DataprepEditor/Private/Widgets/DataprepGraph/SDataprepGraphEditor.cpp
@@ -226,28 +226,10 @@ void SDataprepGraphEditor::CacheDesiredSize(float InLayoutScaleMultiplier)
 		{
 			bIsComplete = TrackGraphNodePtr.Pin()->RefreshLayout();
 			bMustRearrange = true;
-			// Force a change of viewpoint to update the canvas.
-			SetViewLocation(FVector2D(0.f, -TopPadding), 1.f);
 		}
 	}
 }
 
-void SDataprepGraphEditor::UpdateBoundaries(const FVector2D& LocalSize, float ZoomAmount)
-{
-	if(SDataprepGraphTrackNode* TrackGraphNode = TrackGraphNodePtr.Pin().Get())
-	{
-		CachedTrackNodeSize = TrackGraphNode->Update(LocalSize, ZoomAmount);
-	}
-
-	ViewLocationRangeOnY.Set( -TopPadding, -TopPadding );
-
-	const float DesiredVisualHeight = CachedTrackNodeSize.Y * ZoomAmount;
-	if(LocalSize.Y < DesiredVisualHeight)
-	{
-		ViewLocationRangeOnY.Y = DesiredVisualHeight - LocalSize.Y;
-	}
-}
-
 void SDataprepGraphEditor::Tick(const FGeometry& AllottedGeometry, const double InCurrentTime, const float InDeltaTime)
 {
 	// Do not change the layout until all widgets have been created.
@@ -277,70 +259,57 @@ void SDataprepGraphEditor::Tick(const FGeometry& AllottedGeometry, const double
 
 void SDataprepGraphEditor::UpdateLayout( const FVector2D& LocalSize, const FVector2D& Location, float ZoomAmount )
 {
-	if(LastZoomAmount != ZoomAmount)
-	{
-		UpdateBoundaries( LocalSize, ZoomAmount );
-	}
-
-	if( !LocalSize.Equals(LastLocalSize) )
+	if(SDataprepGraphTrackNode* TrackGraphNode = TrackGraphNodePtr.Pin().Get())
 	{
-		bMustRearrange = true;
+		if(LastZoomAmount != ZoomAmount)
+		{
+			WorkingArea = TrackGraphNode->Update();
+		}
 
-		UpdateBoundaries( LocalSize, ZoomAmount );
+		if( !LocalSize.Equals(LastLocalSize) )
+		{
+			bMustRearrange = true;
 
-		LastLocalSize = LocalSize;
+			WorkingArea = TrackGraphNode->Update();
 
-		// Force a re-compute of the view location
-		LastLocation = -Location;
-	}
+			LastLocalSize = LocalSize;
 
-	if( !Location.Equals(LastLocation) )
-	{
-		FVector2D ComputedLocation( LastLocation );
-
-		if(Location.X != LastLocation.X)
-		{
-			const float ActualWidth = LocalSize.X / ZoomAmount;
-			const float MaxInX = CachedTrackNodeSize.X > ActualWidth ? CachedTrackNodeSize.X - ActualWidth : 0.f;
-			ComputedLocation.X = Location.X < 0.f ? 0.f : Location.X >= MaxInX ? MaxInX : Location.X;
+			// Force a re-compute of the view location
+			LastLocation = -Location;
 		}
 
-		if(Location.Y != LastLocation.Y)
+		if( !Location.Equals(LastLocation) )
 		{
-			// Keep same visual Y position if only zoom has changed
-			// Assumption: user cannot zoom in or out and move the canvas at the same time
-			if(LastZoomAmount != ZoomAmount)
+			FVector2D ComputedLocation(Location);
+
+			FVector2D PanelSize = LocalSize / ZoomAmount;
+			FVector2D WorkingSize = WorkingArea.GetSize();
+
+			if(Location.X != LastLocation.X)
 			{
-				ComputedLocation.Y = LastLocation.Y * LastZoomAmount / ZoomAmount;
+				const float Delta = WorkingSize.X - PanelSize.X;
+				const float MaxRight =  Delta > 0.f ? WorkingArea.Left + Delta : WorkingArea.Left;
+				ComputedLocation.X = ComputedLocation.X < WorkingArea.Left ? WorkingArea.Left : (ComputedLocation.X >= MaxRight ? MaxRight : ComputedLocation.X);
 			}
-			else
+
+			if(Location.Y != LastLocation.Y)
 			{
-				const float ActualPositionInY = Location.Y * ZoomAmount;
-				if(ActualPositionInY <= ViewLocationRangeOnY.X )
-				{
-					ComputedLocation.Y = ViewLocationRangeOnY.X / ZoomAmount;
-				}
-				else if( ActualPositionInY > ViewLocationRangeOnY.Y )
-				{
-					ComputedLocation.Y = ViewLocationRangeOnY.Y / ZoomAmount;
-				}
-				else
-				{
-					ComputedLocation.Y = Location.Y;
-				}
+				const float Delta = WorkingSize.Y - PanelSize.Y;
+				const float MaxBottom =  Delta > 0.f ? WorkingArea.Top + Delta : WorkingArea.Top;
+				ComputedLocation.Y = ComputedLocation.Y < WorkingArea.Top ? WorkingArea.Top : (ComputedLocation.Y >= MaxBottom ? MaxBottom : ComputedLocation.Y);
 			}
-		}
 
-		LastLocation = Location;
+			LastLocation = Location;
 
-		if(ComputedLocation != Location)
-		{
-			SetViewLocation( ComputedLocation, ZoomAmount );
-			LastLocation = ComputedLocation;
+			if(ComputedLocation != Location)
+			{
+				SetViewLocation( ComputedLocation, ZoomAmount );
+				LastLocation = ComputedLocation;
+			}
 		}
-	}
 
-	LastZoomAmount = ZoomAmount;
+		LastZoomAmount = ZoomAmount;
+	}
 }
 
 void SDataprepGraphEditor::OnDragEnter(const FGeometry & MyGeometry, const FDragDropEvent & DragDropEvent)
diff --git a/Engine/Plugins/Enterprise/DataprepEditor/Source/DataprepEditor/Private/Widgets/DataprepGraph/SDataprepGraphEditor.h b/Engine/Plugins/Enterprise/DataprepEditor/Source/DataprepEditor/Private/Widgets/DataprepGraph/SDataprepGraphEditor.h
--- a/Engine/Plugins/Enterprise/DataprepEditor/Source/DataprepEditor/Private/Widgets/DataprepGraph/SDataprepGraphEditor.h
+++ b/Engine/Plugins/Enterprise/DataprepEditor/Source/DataprepEditor/Private/Widgets/DataprepGraph/SDataprepGraphEditor.h
@@ -13,6 +13,7 @@
 #include "Framework/Commands/UICommandList.h"
 #include "GraphEditor.h"
 #include "GraphEditorActions.h"
+#include "Layout/SlateRect.h"
 #include "NodeFactory.h"
 #include "SGraphNode.h"
 #include "Widgets/DeclarativeSyntaxSupport.h"
@@ -131,9 +132,6 @@ private:
 	/** Recompute the layout of the displayed graph after a pan, resize and/or zoom */
 	void UpdateLayout( const FVector2D& LocalSize, const FVector2D& Location, float ZoomAmount );
 
-	/** Recompute the boundaries of a the displayed graph */
-	void UpdateBoundaries(const FVector2D& LocalSize, float ZoomAmount);
-
 	void BuildCommandList();
 
 private:
@@ -152,11 +150,8 @@ private:
 	/** Last zoom factor applied to the graph's canvas */
 	float LastZoomAmount;
 
-	/** Indicates min and max of ordinates in canvas */
-	FVector2D ViewLocationRangeOnY;
-
-	/** Size of graph being displayed */
-	mutable FVector2D CachedTrackNodeSize;
+	/** Dimensions of pipeline being displayed */
+	FSlateRect WorkingArea;
 
 	/** Size of graph being displayed */
 	TWeakObjectPtr<UDataprepAsset> DataprepAssetPtr;
@@ -178,7 +173,7 @@ private:
 	/** Factory to create the associated SGraphNode classes for Dataprep graph's UEdGraph classes */
 	static TSharedPtr<FDataprepGraphEditorNodeFactory> NodeFactory;
 
-	// The dataprep editor that owns the graph
+	// The Dataprep editor that owns the graph
 	TWeakPtr<FDataprepEditor> DataprepEditor;
 
 	friend SDataprepGraphTrackNode;
diff --git a/Engine/Plugins/Enterprise/DataprepEditor/Source/DataprepEditor/Private/Widgets/DataprepGraph/SDataprepGraphTrackNode.cpp b/Engine/Plugins/Enterprise/DataprepEditor/Source/DataprepEditor/Private/Widgets/DataprepGraph/SDataprepGraphTrackNode.cpp
--- a/Engine/Plugins/Enterprise/DataprepEditor/Source/DataprepEditor/Private/Widgets/DataprepGraph/SDataprepGraphTrackNode.cpp
+++ b/Engine/Plugins/Enterprise/DataprepEditor/Source/DataprepEditor/Private/Widgets/DataprepGraph/SDataprepGraphTrackNode.cpp
@@ -19,6 +19,7 @@
 #include "SGraphPanel.h"
 #include "NodeFactory.h"
 #include "Widgets/Colors/SColorBlock.h"
+#include "Widgets/Images/SImage.h"
 #include "Widgets/Layout/SBox.h"
 #include "Widgets/Layout/SConstraintCanvas.h"
 #include "Widgets/SBoxPanel.h"
@@ -26,47 +27,109 @@
 
 #define LOCTEXT_NAMESPACE "DataprepGraphEditor"
 
-FMargin SDataprepGraphTrackNode::NodePadding( 15.f, -5.f, 5.f, 10.f );
+FVector2D SDataprepGraphTrackNode::TrackAnchor( 15.f, 40.f );
 
-class SDataprepGraphTrackWidget : public SHorizontalBox
+class SDataprepEmptyActionNode : public SVerticalBox
+{
+public:
+	SLATE_BEGIN_ARGS(SDataprepEmptyActionNode){}
+	SLATE_END_ARGS();
+
+	void Construct(const FArguments& InArgs)
+	{
+		SVerticalBox::Construct(SVerticalBox::FArguments());
+
+		AddSlot()
+		.AutoHeight()
+		.Padding(FDataprepEditorStyle::GetMargin( "DataprepActionStep.Padding" ))
+		[
+			SNew(SOverlay)
+
+			+ SOverlay::Slot()
+			.Padding(10.f)
+			.VAlign(VAlign_Fill)
+			.HAlign(HAlign_Fill)
+			[
+				SNew(SBox)
+				.MinDesiredWidth(250.f)
+			]
+
+			+ SOverlay::Slot()
+			.Padding(0.f)
+			.VAlign(VAlign_Fill)
+			.HAlign(HAlign_Fill)
+			[
+				SNew(SImage)
+				.ColorAndOpacity(FDataprepEditorStyle::GetColor("DataprepAction.EmptyStep.Outer.Hovered"))
+				.Image(FEditorStyle::GetBrush( "Graph.StateNode.Body" ))
+			]
+
+			+ SOverlay::Slot()
+			.Padding(1.f)
+			.VAlign(VAlign_Fill)
+			.HAlign(HAlign_Fill)
+			[
+				SNew(SImage)
+				.ColorAndOpacity(FDataprepEditorStyle::GetColor("DataprepAction.EmptyStep.Background.Hovered"))
+				.Image(FEditorStyle::GetBrush( "Graph.StateNode.Body" ))
+			]
+
+			+ SOverlay::Slot()
+			.Padding(10.f)
+			.VAlign(VAlign_Fill)
+			.HAlign(HAlign_Fill)
+			[
+				SNew(SHorizontalBox)
+
+				+ SHorizontalBox::Slot()
+				.VAlign(VAlign_Fill)
+				.HAlign(HAlign_Fill)
+				[
+					SNew(STextBlock)
+					.Text(LOCTEXT("DataprepEmptyActionLabel", "+ Add Action"))
+					.TextStyle( &FDataprepEditorStyle::GetWidgetStyle<FTextBlockStyle>( "DataprepAction.TitleTextBlockStyle" ) )
+					.ColorAndOpacity(FDataprepEditorStyle::GetColor("DataprepAction.EmptyStep.Text.Hovered"))
+					.Justification(ETextJustify::Center)
+				]
+			]
+		];
+	}
+};
+
+class SDataprepGraphTrackWidget : public SConstraintCanvas
 {
 public:
 	SLATE_BEGIN_ARGS(SDataprepGraphTrackWidget) {}
 	SLATE_END_ARGS();
 
-	void Construct(const FArguments& InArgs, TSharedPtr<SDataprepGraphTrackNode> InTrackNode);
 
-	void MoveDropSlotTo(int32 SlotIndex);
+	void Construct(const FArguments& InArgs, TSharedPtr<SDataprepGraphTrackNode> InTrackNode);
 
 	void SetActionSlot(int32 SlotIndex, const TSharedRef<SWidget>& Widget);
 
-	// SWidget Interface
-	virtual FVector2D ComputeDesiredSize(float InLayoutScaleMultiplier) const override;
-	// End of SWidget Interface
-
-	const FVector2D& GetWorkingArea() const
+	const FVector2D& GetWorkingSize() const
 	{
-		return WorkingArea;
+		return WorkingSize;
 	}
 
-	FVector2D GetWorkingArea()
+	FVector2D GetWorkingSize()
 	{
-		return WorkingArea;
+		return WorkingSize;
 	}
 
 	FVector2D GetTrackArea()
 	{
-		return FVector2D(WorkingArea.X, TrackDesiredHeight);
+		return FVector2D(WorkingSize.X, TrackDesiredHeight);
 	}
 
-	const FVector2D& GetCanvasArea() const
+	FVector2D GetCanvasArea() const
 	{
-		return CanvasArea;
+		return FVector2D(CanvasOffset.Right, CanvasOffset.Bottom);
 	}
 
-	FVector2D GetCanvasArea()
+	FVector2D GetLineSize() const
 	{
-		return CanvasArea;
+		return LineSize;
 	}
 
 	FVector2D GetActioNodeAbscissaRange();
@@ -78,11 +141,9 @@ public:
 		return ActionSlots.IsValidIndex(Index) ? ActionSlots[Index]->OffsetAttr.Get(FMargin()).Left : 0.f;
 	}
 
-	float ValidateNodeAbscissa(float InAbscissa);
-
 	FVector2D GetNodeSize(const TSharedRef<SWidget>& Widget) const;
 
-	void RefreshLayout(SGraphPanel* GraphPanel);
+	void RefreshLayout();
 
 	void OnStartNodeDrag(int32 Index);
 
@@ -94,61 +155,106 @@ public:
 
 	/**
 	*/
-	void UpdateLayoutForCopy();
-
-	/**
-	*/
-	void UpdateLayoutForMove();
+	void UpdateLayoutOnDrag(int8 DragDirection = 0);
 
 	/** Toggles display of action nodes between copy vs move drop modes */
 	float OnControlKeyDown(bool bKeyDown);
 
 	void UpdateDragIndicator(FVector2D MousePosition);
 
-	int32 GetDragIndicatorIndex();
+	int32 GetDragIndicatorIndex() { return DragIndicatorIndex; }
 
 	void ResetDragIndicator()
 	{
-		for(int32 Index = 0; Index < DropSlots.Num(); ++Index)
-		{
-			(*DropSlots[Index])[DropFiller.ToSharedRef()];
-		}
+		DragIndicatorIndex = INDEX_NONE;
+
+		FMargin SlotOffset = ActionSlots[SlotCount]->OffsetAttr.Get(FMargin());
+		TrackSlots[SlotCount]->OffsetAttr.Set( FMargin(SlotOffset.Left + SDataprepGraphActionNode::DefaultWidth * 0.5f - 16.f, LineTopPadding + TrackSlotTopOffset, 32, 32));
+		TrackSlots[SlotCount]->AttachWidget(SNullWidget::NullWidget);
+
+		UpdateLine(SlotCount-1);
+
+	}
+
+	FVector2D UpdateActionSlot(FMargin& SlotOffset, int32 SlotIndex, const TSharedRef<SWidget>& ActionNode, bool bIsEmpty = false)
+	{
+		DropSlots[SlotIndex]->OffsetAttr.Set(SlotOffset);
+		SlotOffset.Left += InterSlotSpacing;
+
+		const float NodeTopPadding = LineTopPadding + NodeTopOffset;
+		const float TrackSlotTopPadding = LineTopPadding + TrackSlotTopOffset;
+
+		SConstraintCanvas::FSlot* ActionSlot = ActionSlots[SlotIndex];
+
+		FVector2D Size = GetNodeSize(ActionNode);
+
+		ActionSlot->OffsetAttr.Set(FMargin(SlotOffset.Left, NodeTopPadding, Size.X, Size.Y));
+		ActionSlot->AttachWidget(bIsEmpty ? SNew(SColorBlock).Color(FLinearColor::Transparent).Size( Size ) : ActionNode);
+
+		TrackSlots[SlotIndex]->OffsetAttr.Set( FMargin(SlotOffset.Left + Size.X * 0.5f - 16.f, TrackSlotTopPadding, 32, 32));
+		TrackSlots[SlotIndex]->AttachWidget(SlotIndex == HoveredDragIndex ? TrackSlotDrop.ToSharedRef() : TrackSlotRegular.ToSharedRef());
+
+		SlotOffset.Left += Size.X;
+
+		return Size + FVector2D(0.f, NodeTopPadding);
 	}
 
 	void OnActionsOrderChanged(const TArray<TSharedPtr<SDataprepGraphActionNode>>& InActionNodes);
 
+	void CreateHelperWidgets();
+
+	void GetDraggedNodePositions(FVector2D& Left, FVector2D& Right, const FGeometry& TargetGeometry);
+
+	void UpdateLine(int32 LastIndex);
+
 	TWeakPtr<SDataprepGraphTrackNode> TrackNodePtr;
 	TSharedPtr<SConstraintCanvas> TrackCanvas;
 	TArray<TSharedRef<SWidget>> ActionNodes;
 	TArray<SConstraintCanvas::FSlot*> DropSlots;
 	TArray<SConstraintCanvas::FSlot*> ActionSlots;
+	TArray<SConstraintCanvas::FSlot*> TrackSlots;
 	SConstraintCanvas::FSlot* CanvasSlot;
+	SConstraintCanvas::FSlot* LineSlot;
+
+	FVector2D WorkingSize;
+	FMargin CanvasOffset;
+	FVector2D LineSize;
+	int32 LastDropSlot;
 
 	TSharedPtr<SWidget> DropIndicator;
 	TSharedPtr<SColorBlock> DropFiller;
-	mutable FVector2D WorkingArea;
-	mutable FVector2D CanvasArea;
-	int32 LastDropSlot;
+	TSharedPtr<SWidget> TrackSlotRegular;
+	TSharedPtr<SWidget> TrackSlotDrop;
+	TSharedPtr<SWidget> TrackSlotSelected;
+	TSharedPtr<SDataprepEmptyActionNode> DummyAction;
+
+	/** Number of visible slots. THis number can only change when dragging actions */
+	int32 SlotCount;
 
 	SConstraintCanvas::FSlot* DragSlot;
-	TArray<int32> NewActionsOrder;
 	int32 DragStartIndex;
 	int32 DragIndex;
+	int32 HoveredDragIndex;
+	int32 DragIndicatorIndex;
 	FVector2D LastDragPosition;
 	FVector2D AbscissaRange;
 	FVector2D TrackOffset;
 	bool bIsCopying;
 
-	static float SlotSpacing;
-	static float NodeDesiredWidth;
 	static float TrackDesiredHeight;
+	static float InterSlotSpacing;
+	static float LineTopPadding;
+	static float NodeTopOffset;
+	static float TrackSlotTopOffset;
 
 	friend SDataprepGraphTrackNode;
 };
 
-float SDataprepGraphTrackWidget::SlotSpacing = 16.f;
-float SDataprepGraphTrackWidget::NodeDesiredWidth = 300.f;
+float SDataprepGraphTrackWidget::InterSlotSpacing = 16.f;
 float SDataprepGraphTrackWidget::TrackDesiredHeight = 40.f;
+float SDataprepGraphTrackWidget::LineTopPadding = 10.f;
+float SDataprepGraphTrackWidget::NodeTopOffset = 25.f;
+float SDataprepGraphTrackWidget::TrackSlotTopOffset = -15.f;
 
 void SDataprepGraphTrackNode::Construct(const FArguments& InArgs, UDataprepGraphRecipeNode* InNode)
 {
@@ -166,7 +272,7 @@ void SDataprepGraphTrackNode::Construct(const FArguments& InArgs, UDataprepGraph
 	check(DataprepAssetPtr.IsValid());
 
 	SNodePanel::SNode::FNodeSet NodeFilter;
-	SGraphNode::MoveTo( FVector2D::ZeroVector, NodeFilter);
+	SGraphNode::MoveTo( TrackAnchor, NodeFilter);
 
 	InNode->SetWidget(SharedThis(this));
 
@@ -229,7 +335,6 @@ void SDataprepGraphTrackNode::UpdateGraphNode()
 					}
 
 					ActionWidget->UpdateGraphNode();
-					ActionWidget->ComputeDesiredSize(1.0f);
 
 					ActionWidget->SetParentTrackNode(ParentTrackNodePtr);
 
@@ -267,9 +372,6 @@ void SDataprepGraphTrackNode::UpdateGraphNode()
 			]
 		]
 	];
-
-	GraphNode->NodePosX = NodePadding.Left;
-	GraphNode->NodePosY = NodePadding.Top;
 }
 
 bool SDataprepGraphTrackNode::CustomPrepass(float LayoutScaleMultiplier)
@@ -328,14 +430,7 @@ bool SDataprepGraphTrackNode::RefreshLayout()
 	{
 		ensure(TrackWidgetPtr.IsValid());
 
-		const FVector2D TrackWidgetOrigin = TrackWidgetPtr->GetCachedGeometry().LocalToAbsolute(FVector2D::ZeroVector);
-		const FVector2D TrackNodeOrigin = GetCachedGeometry().LocalToAbsolute(FVector2D::ZeroVector);
-
-		// Weird: Ultimate offset between track widget and node is 5 but it may take more draw call to stabilize
-		//		  Forcing 5.
-		TrackWidgetPtr->TrackOffset = FVector2D(5.f)/*TrackWidgetOrigin - TrackNodeOrigin*/ + GetPosition();
-
-		TrackWidgetPtr->RefreshLayout(OwnerGraphPanelPtr.Pin().Get());
+		TrackWidgetPtr->RefreshLayout();
 
 		return true;
 	}
@@ -346,7 +441,7 @@ bool SDataprepGraphTrackNode::RefreshLayout()
 void SDataprepGraphTrackNode::MoveTo(const FVector2D& NewPosition, FNodeSet& NodeFilter)
 {
 	// Block track node to specific position
-	SGraphNode::MoveTo( FVector2D(NodePadding.Left, 0.f), NodeFilter);
+	SGraphNode::MoveTo( TrackAnchor, NodeFilter);
 }
 
 const FSlateBrush* SDataprepGraphTrackNode::GetShadowBrush(bool bSelected) const
@@ -364,13 +459,13 @@ FReply SDataprepGraphTrackNode::OnDrop(const FGeometry& MyGeometry, const FDragD
 		return FReply::Handled().EndDragDrop();
 	}
 
-	TSharedPtr<FDataprepDragDropOp> DragActionStepNodeOp = DragDropEvent.GetOperationAs<FDataprepDragDropOp>();
-	if(DragActionStepNodeOp.IsValid())
+	TSharedPtr<FDataprepDragDropOp> DragActionNodeOp = DragDropEvent.GetOperationAs<FDataprepDragDropOp>();
+	if(DragActionNodeOp.IsValid())
 	{
 		int32 InsertIndex = TrackWidgetPtr->GetDragIndicatorIndex();
 		TrackWidgetPtr->ResetDragIndicator();
 
-		DragActionStepNodeOp->DoDropOnTrack(DataprepAssetPtr.Get(), InsertIndex);
+		DragActionNodeOp->DoDropOnTrack(DataprepAssetPtr.Get(), InsertIndex);
 
 		return FReply::Unhandled();
 	}
@@ -380,10 +475,10 @@ FReply SDataprepGraphTrackNode::OnDrop(const FGeometry& MyGeometry, const FDragD
 
 FReply SDataprepGraphTrackNode::OnDragOver(const FGeometry & MyGeometry, const FDragDropEvent & DragDropEvent)
 {
-	TSharedPtr<FDataprepDragDropOp> DragActionStepNodeOp = DragDropEvent.GetOperationAs<FDataprepDragDropOp>();
-	if(DragActionStepNodeOp.IsValid())
+	TSharedPtr<FDataprepDragDropOp> DragActionNodeOp = DragDropEvent.GetOperationAs<FDataprepDragDropOp>();
+	if(DragActionNodeOp.IsValid())
 	{
-		DragActionStepNodeOp->SetHoveredNode(GraphNode);
+		DragActionNodeOp->SetHoveredNode(GraphNode);
 		TrackWidgetPtr->UpdateDragIndicator(DragDropEvent.GetScreenSpacePosition());
 	}
 
@@ -392,10 +487,11 @@ FReply SDataprepGraphTrackNode::OnDragOver(const FGeometry & MyGeometry, const F
 
 void SDataprepGraphTrackNode::OnDragLeave(const FDragDropEvent & DragDropEvent)
 {
-	TSharedPtr<FDataprepDragDropOp> DragActionStepNodeOp = DragDropEvent.GetOperationAs<FDataprepDragDropOp>();
-	if(DragActionStepNodeOp.IsValid())
+	TSharedPtr<FDataprepDragDropOp> DragActionNodeOp = DragDropEvent.GetOperationAs<FDataprepDragDropOp>();
+	if(DragActionNodeOp.IsValid())
 	{
 		TrackWidgetPtr->ResetDragIndicator();
+		TrackWidgetPtr->RefreshLayout();
 	}
 
 	SGraphNode::OnDragLeave(DragDropEvent);
@@ -416,16 +512,45 @@ void SDataprepGraphTrackNode::SetOwner(const TSharedRef<SGraphPanel>& OwnerPanel
 	}
 }
 
-FVector2D SDataprepGraphTrackNode::Update(const FVector2D& LocalSize, float ZoomAmount)
+FSlateRect SDataprepGraphTrackNode::Update()
 {
+	FSlateRect WorkingArea;
+
 	if(SDataprepGraphTrackWidget* TrackWidget = TrackWidgetPtr.Get())
 	{
 		RefreshLayout();
 
-		return TrackWidget->GetWorkingArea();
+		SGraphPanel& GraphPanel = *OwnerGraphPanelPtr.Pin();
+
+		const FVector2D& WorkingSize = TrackWidget->WorkingSize;
+
+		// Determine canvas offset attribute in track widget's coordinates
+		const FVector2D PanelSize = GraphPanel.GetTickSpaceGeometry().GetLocalSize() / GraphPanel.GetZoomAmount();
+
+		const FVector2D TrackPosition = GetPosition().GetAbs();
+		const FVector2D TargetSize = FVector2D(FMath::Max(WorkingSize.X, PanelSize.X), FMath::Max(WorkingSize.Y, PanelSize.Y)) + TrackPosition + 20.f;
+
+		FMargin& CanvasOffset = TrackWidgetPtr->CanvasOffset;
+
+		CanvasOffset.Left = -TrackPosition.X - 10.f;
+		CanvasOffset.Top = -TrackPosition.Y - 10.f;
+
+		if(TargetSize.X > CanvasOffset.Right)
+		{
+			CanvasOffset.Right = TargetSize.X;
+		}
+
+		if(TargetSize.Y > CanvasOffset.Bottom)
+		{
+			CanvasOffset.Bottom = TargetSize.Y;
+		}
+
+		TrackWidgetPtr->CanvasSlot->OffsetAttr.Set(CanvasOffset);
+
+		WorkingArea = FSlateRect(FVector2D::ZeroVector, WorkingSize + TrackPosition);
 	}
 
-	return FVector2D(SDataprepGraphTrackWidget::NodeDesiredWidth);
+	return WorkingArea;
 }
 
 void SDataprepGraphTrackNode::OnStartNodeDrag(const TSharedRef<SDataprepGraphActionNode>& ActionNode)
@@ -472,8 +597,9 @@ void SDataprepGraphTrackNode::OnNodeDragged( TSharedPtr<SDataprepGraphActionNode
 
 	TSharedPtr<SGraphPanel> GraphPanel = GetOwnerPanel();
 	const FGeometry& DesktopGeometry = GraphPanel->GetCachedGeometry();
-	const FVector2D& Size = DesktopGeometry.GetLocalSize();
+	const FVector2D& PanelSize = DesktopGeometry.GetLocalSize();
 	const float ZoomAmount = GraphPanel->GetZoomAmount();
+	const FVector2D& TrackSize = TrackWidgetPtr->GetWorkingSize() * ZoomAmount;
 
 	float AbscissaDelta = TrackWidgetPtr->OnNodeDragged(ScreenSpaceDelta.X / ZoomAmount);
 		
@@ -487,16 +613,22 @@ void SDataprepGraphTrackNode::OnNodeDragged( TSharedPtr<SDataprepGraphActionNode
 			// Do nothing, the mouse's cursor and the dragged node should not move
 		}
 		// Panel is narrower than track, update panel if dragged node has entered non visible part of track
-		else if((TrackWidgetPtr->GetWorkingArea().X * ZoomAmount) > Size.X)
+		else if(TrackSize.X > PanelSize.X)
 		{
-			const FGeometry& NodeGeometry = ActionNodePtr->GetCachedGeometry();
-			FVector2D NodeAbsolutePosition = NodeGeometry.LocalToAbsolute(FVector2D(NodeGeometry.GetLocalSize().X, 0.f));
-			FVector2D NodeRelativePosition = DesktopGeometry.AbsoluteToLocal(NodeAbsolutePosition);
+			FVector2D LeftPosition;
+			FVector2D RightPosition;
+			TrackWidgetPtr->GetDraggedNodePositions(LeftPosition, RightPosition, DesktopGeometry);
 
-			if(NodeRelativePosition.X > Size.X)
+			if(RightPosition.X > PanelSize.X)
 			{
 				FVector2D ViewOffset = GraphPanel->GetViewOffset();
-				ViewOffset.X += NodeRelativePosition.X - DesktopGeometry.GetLocalSize().X;
+				ViewOffset.X += RightPosition.X - PanelSize.X + 10.f;
+				GraphPanel->RestoreViewSettings(ViewOffset, ZoomAmount);
+			}
+			else if(LeftPosition.X < 0.f)
+			{
+				FVector2D ViewOffset = GraphPanel->GetViewOffset();
+				ViewOffset.X += LeftPosition.X - 10.f;
 				GraphPanel->RestoreViewSettings(ViewOffset, ZoomAmount);
 			}
 			else
@@ -628,27 +760,37 @@ void SDataprepGraphTrackWidget::Construct(const FArguments& InArgs, TSharedPtr<S
 	TrackNodePtr = InTrackNode;
 	ensure(TrackNodePtr.IsValid());
 
-	SHorizontalBox::Construct( SHorizontalBox::FArguments()
-		+SHorizontalBox::Slot()
-		.AutoWidth()
-		[
-			SAssignNew( TrackCanvas, SConstraintCanvas )
-		]
-	);
+	SConstraintCanvas::Construct( SConstraintCanvas::FArguments());
+
+	CreateHelperWidgets();
+
+	const FLinearColor TrackColor = FDataprepEditorStyle::GetColor( "DataprepAction.BackgroundColor" );
 
-	CanvasSlot = &TrackCanvas->AddSlot()
+	CanvasSlot = &AddSlot()
 	.Anchors( FAnchors( 0.f, 0.f, 0.f, 0.f ) )
-	.Offset( FMargin(-20.f, -50.f, 0.f, 0.f) )
-	.AutoSize(true)
+	.Offset( FMargin(-25.f, -25.f, 0.f, 0.f) )
+	.AutoSize(false)
 	.Alignment(FVector2D::ZeroVector)
 	.ZOrder(0)
 	[
 		SNew(SColorBlock)
-		.Color( FLinearColor::Transparent )
+		.Color( FDataprepEditorStyle::GetColor("Dataprep.Background.Black") )
 		.Size( TAttribute<FVector2D>::Create(TAttribute<FVector2D>::FGetter::CreateSP( this, &SDataprepGraphTrackWidget::GetCanvasArea ) ) )
 	];
 
-	TrackCanvas->AddSlot()
+	LineSlot = &AddSlot()
+	.Anchors( FAnchors( 0.f, 0.f, 0.f, 0.f ) )
+	.Offset( FMargin(InterSlotSpacing +  3.f, LineTopPadding - 2.f, 0.f, 0.f) )
+	.AutoSize(true)
+	.Alignment(FVector2D::ZeroVector)
+	.ZOrder(0)
+	[
+		SNew(SColorBlock)
+		.Color( TrackColor )
+		.Size( TAttribute<FVector2D>::Create(TAttribute<FVector2D>::FGetter::CreateSP( this, &SDataprepGraphTrackWidget::GetLineSize ) ) )
+	];
+
+	AddSlot()
 	.Anchors( FAnchors( 0.f, 0.f, 1.f, 0.f ) )
 	.Offset( FMargin(0.f, 10.f, 0.f, 0.f) )
 	.AutoSize(true)
@@ -656,50 +798,43 @@ void SDataprepGraphTrackWidget::Construct(const FArguments& InArgs, TSharedPtr<S
 	.ZOrder(1)
 	[
 		SNew(SColorBlock)
-		.Color(FDataprepEditorStyle::GetColor( "Graph.TrackInner.BackgroundColor" ))
+		.Color(FLinearColor::Transparent)
 		.Size( TAttribute<FVector2D>::Create(TAttribute<FVector2D>::FGetter::CreateSP( this, &SDataprepGraphTrackWidget::GetTrackArea ) ) )
 	];
 
 	LastDropSlot = INDEX_NONE;
 
-	DropIndicator = SNew(SBorder)
-	.BorderBackgroundColor( FDataprepEditorStyle::GetColor( "DataprepActionStep.DragAndDrop" ) )
-	.BorderImage(FEditorStyle::GetBrush("Graph.StateNode.Body"))
-	.HAlign(EHorizontalAlignment::HAlign_Center)
-	[
-		SNew(SBox)
-		.HeightOverride(TrackDesiredHeight - 2.f)
-		.WidthOverride(SlotSpacing - 2.f)
-	];
-
-	DropFiller = SNew(SColorBlock)
-	.Color(FLinearColor::Transparent)
-	.Size( FVector2D(SlotSpacing, TrackDesiredHeight) );
-
 	TSharedPtr<SColorBlock> NullNode = SNew(SColorBlock)
 	.Color( FLinearColor::Transparent )
-	.Size( FVector2D(NodeDesiredWidth, TrackDesiredHeight) );
+	.Size( FVector2D(SDataprepGraphActionNode::DefaultWidth, TrackDesiredHeight) );
+
+	const FVector2D DefaultNodeSize(SDataprepGraphActionNode::DefaultWidth, SDataprepGraphActionNode::DefaultHeight);
+
+	CanvasOffset = FMargin(FVector4(FVector2D::ZeroVector, DefaultNodeSize));
+
+	const float NodeTopPadding = LineTopPadding + NodeTopOffset;
+	const float TrackSlotTopPadding = LineTopPadding + TrackSlotTopOffset;
 
-	CanvasArea = FVector2D(NodeDesiredWidth, NodeDesiredWidth);
+	UDataprepAsset* DataprepAsset = InTrackNode->GetDataprepAsset();
+	SlotCount = DataprepAsset ? DataprepAsset->GetActionCount() : 0;
+
+	DropSlots.Reserve(SlotCount + 2);
+	ActionSlots.Reserve(SlotCount + 1);
+	TrackSlots.Reserve(SlotCount + 2);
+	ActionNodes.Reserve(SlotCount);
 
 	float LeftOffset = 0.f;
-	if(UDataprepAsset* DataprepAsset = InTrackNode->GetDataprepAsset())
+	if(DataprepAsset)
 	{
-		int32 ActionCount = DataprepAsset->GetActionCount();
-
 		TArray<TSharedPtr<SDataprepGraphActionNode>>& InActionNodes = InTrackNode->ActionNodes;
 
-		DropSlots.Reserve(ActionCount + 1);
-		ActionSlots.Reserve(ActionCount);
-		ActionNodes.Reserve(ActionCount);
-
-		for(int32 Index = 0; Index < ActionCount; ++Index)
+		for(int32 Index = 0; Index < SlotCount; ++Index)
 		{
 			// Add drop slot ahead of action node
 			{
-				SConstraintCanvas::FSlot& Slot = TrackCanvas->AddSlot()
+				SConstraintCanvas::FSlot& Slot = AddSlot()
 				.Anchors( FAnchors( 0., 0.f, 0.f, 0.f ) )
-				.Offset( FMargin(LeftOffset, 10.f, SlotSpacing, TrackDesiredHeight ) )
+				.Offset( FMargin(LeftOffset, NodeTopPadding + 5.f, InterSlotSpacing, TrackDesiredHeight ) )
 				.Alignment(FVector2D::ZeroVector)
 				.ZOrder(2)
 				[
@@ -707,48 +842,68 @@ void SDataprepGraphTrackWidget::Construct(const FArguments& InArgs, TSharedPtr<S
 				];
 				DropSlots.Add( &Slot );
 			}
-			LeftOffset += SlotSpacing;
+			LeftOffset += InterSlotSpacing;
+
+			const FVector2D Size = DefaultNodeSize;
 
 			// Add slot to host of action node
 			{
-				ActionNodes.Add(InActionNodes[Index].IsValid() ? InActionNodes[Index]->AsShared() : NullNode->AsShared());
+				TSharedRef<SWidget>& ActionNode = ActionNodes.Add_GetRef(InActionNodes[Index].IsValid() ? InActionNodes[Index]->AsShared() : NullNode->AsShared());
 
-				SConstraintCanvas::FSlot& Slot = TrackCanvas->AddSlot()
+				SConstraintCanvas::FSlot& Slot = AddSlot()
 				.Anchors( FAnchors( 0., 0.f, 0.f, 0.f ) )
-				.Offset( FMargin(LeftOffset, 0.f, NodeDesiredWidth, NodeDesiredWidth ) )
+				.Offset( FMargin(LeftOffset, NodeTopPadding, Size.X, Size.Y ) )
 				.Alignment(FVector2D::ZeroVector)
 				.AutoSize(true)
 				.ZOrder(2)
 				[
-					ActionNodes[Index]
+					ActionNode
 				];
 				ActionSlots.Add( &Slot );
 			}
-			LeftOffset += NodeDesiredWidth;
+
+			// Add track slot which is at the middle of the corresponding action
+			{
+				SConstraintCanvas::FSlot& Slot = AddSlot()
+				.Anchors( FAnchors( 0., 0.f, 0.f, 0.f ) )
+				.Offset( FMargin(LeftOffset + SDataprepGraphActionNode::DefaultWidth * 0.5f + 16.f, TrackSlotTopPadding, 32, 32 ) )
+				.Alignment(FVector2D::ZeroVector)
+				.AutoSize(true)
+				.ZOrder(2)
+				[
+					TrackSlotRegular.ToSharedRef()
+				];
+				TrackSlots.Add( &Slot );
+			}
+			LeftOffset += Size.X;
+
+			CanvasOffset.Bottom = CanvasOffset.Bottom < Size.Y ? Size.Y : CanvasOffset.Bottom;
 		}
 	}
-
-	CanvasArea.X += LeftOffset + TrackDesiredHeight;
+	else
+	{
+		SlotCount = 0;
+	}
 
 	// Add drop slot at end of track
 	{
-		SConstraintCanvas::FSlot& Slot = TrackCanvas->AddSlot()
+		SConstraintCanvas::FSlot& Slot = AddSlot()
 		.Anchors( FAnchors( 0., 0.f, 0.f, 0.f ) )
-		.Offset( FMargin(LeftOffset, 5.f, SlotSpacing, TrackDesiredHeight ) )
+		.Offset( FMargin(LeftOffset, NodeTopPadding + 5.f, InterSlotSpacing, TrackDesiredHeight ) )
 		.Alignment(FVector2D::ZeroVector)
 		.ZOrder(2)
 		[
-			DropFiller.ToSharedRef()
+			SNullWidget::NullWidget
 		];
 		DropSlots.Add( &Slot );
-		LeftOffset += SlotSpacing;
+		LeftOffset += InterSlotSpacing;
 	}
 
-	// Add action slot at end of track
+	// Add dummy action slot at end of track
 	{
-		SConstraintCanvas::FSlot& Slot = TrackCanvas->AddSlot()
+		SConstraintCanvas::FSlot& Slot = AddSlot()
 		.Anchors( FAnchors( 0., 0.f, 0.f, 0.f ) )
-		.Offset( FMargin(LeftOffset, 0.f, NodeDesiredWidth, NodeDesiredWidth ) )
+		.Offset( FMargin(LeftOffset, TrackDesiredHeight, SDataprepGraphActionNode::DefaultWidth, SDataprepGraphActionNode::DefaultWidth ) )
 		.Alignment(FVector2D::ZeroVector)
 		.AutoSize(true)
 		.ZOrder(2)
@@ -758,10 +913,39 @@ void SDataprepGraphTrackWidget::Construct(const FArguments& InArgs, TSharedPtr<S
 		ActionSlots.Add( &Slot );
 	}
 
-	// Add slot to host dragged node
-	DragSlot = &TrackCanvas->AddSlot()
+	// Add track slot which is at the middle of the corresponding action
+	{
+		SConstraintCanvas::FSlot& Slot = AddSlot()
+		.Anchors( FAnchors( 0., 0.f, 0.f, 0.f ) )
+		.Offset( FMargin(LeftOffset + SDataprepGraphActionNode::DefaultWidth * 0.5f + 16.f, TrackSlotTopPadding, 32, 32 ) )
+		.Alignment(FVector2D::ZeroVector)
+		.AutoSize(true)
+		.ZOrder(2)
+		[
+			SNullWidget::NullWidget
+		];
+		TrackSlots.Add( &Slot );
+	}
+	LeftOffset += SDataprepGraphActionNode::DefaultWidth;
+
+	// Add drop slot at end of track
+	{
+		SConstraintCanvas::FSlot& Slot = AddSlot()
+		.Anchors( FAnchors( 0., 0.f, 0.f, 0.f ) )
+		.Offset( FMargin(LeftOffset, NodeTopPadding + 5.f, InterSlotSpacing, TrackDesiredHeight ) )
+		.Alignment(FVector2D::ZeroVector)
+		.ZOrder(2)
+		[
+			SNullWidget::NullWidget
+		];
+		DropSlots.Add( &Slot );
+		LeftOffset += InterSlotSpacing;
+	}
+
+	// Add slot to host dragged node???
+	DragSlot = &AddSlot()
 	.Anchors( FAnchors( 0., 0.f, 0.f, 0.f ) )
-	.Offset( FMargin(LeftOffset, 0.f, NodeDesiredWidth, NodeDesiredWidth ) )
+	.Offset( FMargin(LeftOffset, 0.f, SDataprepGraphActionNode::DefaultWidth, SDataprepGraphActionNode::DefaultWidth ) )
 	.Alignment(FVector2D::ZeroVector)
 	.AutoSize(true)
 	.ZOrder(3)
@@ -769,51 +953,101 @@ void SDataprepGraphTrackWidget::Construct(const FArguments& InArgs, TSharedPtr<S
 		SNullWidget::NullWidget
 	];
 
+	WorkingSize = FVector2D::ZeroVector;
+	CanvasOffset.Right = LeftOffset;
+
+	HoveredDragIndex = INDEX_NONE;
 	DragStartIndex = INDEX_NONE;
 	DragIndex = INDEX_NONE;
+	DragIndicatorIndex = INDEX_NONE;
 }
 
-void SDataprepGraphTrackWidget::MoveDropSlotTo(int32 SlotIndex)
+void SDataprepGraphTrackWidget::CreateHelperWidgets()
 {
-	ensure(DropSlots.IsValidIndex(SlotIndex));
-	if(LastDropSlot >= 0)
-	{
-		(*DropSlots[LastDropSlot])[SNullWidget::NullWidget->AsShared()];
-	}
+	const FLinearColor TrackColor = FDataprepEditorStyle::GetColor( "DataprepAction.BackgroundColor" );
+	const FLinearColor DragAndDropColor = FDataprepEditorStyle::GetColor( "DataprepAction.DragAndDrop" );
 
-	LastDropSlot = SlotIndex;
-	if(LastDropSlot >= 0)
-	{
-		(*DropSlots[LastDropSlot])[DropIndicator.ToSharedRef()];
-	}
+	DropIndicator = SNew(SBorder)
+		.BorderBackgroundColor( DragAndDropColor )
+		.BorderImage(FEditorStyle::GetBrush("Graph.StateNode.Body"))
+		.HAlign(EHorizontalAlignment::HAlign_Center)
+		[
+			SNew(SBox)
+			.HeightOverride(TrackDesiredHeight - 2.f)
+			.WidthOverride(InterSlotSpacing - 2.f)
+		];
 
-	Invalidate(EInvalidateWidgetReason::Layout);
+	DropFiller = SNew(SColorBlock)
+		.Color(FLinearColor::Transparent)
+		.Size( FVector2D(InterSlotSpacing, TrackDesiredHeight) );
+
+	TrackSlotRegular = SNew(SOverlay)
+		+SOverlay::Slot()
+		.VAlign(VAlign_Fill)
+		.HAlign(HAlign_Fill)
+		.Padding(0.f)
+		[
+			SNew(SBox)
+			.WidthOverride(32.f)
+			.HeightOverride(32.f)
+		]
+		+SOverlay::Slot()
+		.VAlign(VAlign_Fill)
+		.HAlign(HAlign_Fill)
+		.Padding(0.f)
+		[
+			SNew(SImage)
+			.ColorAndOpacity(TrackColor)
+			.Image(FDataprepEditorStyle::GetBrush( "DataprepEditor.TrackNode.Slot" ))
+		];
+
+	TrackSlotDrop = SNew(SOverlay)
+		+SOverlay::Slot()
+		.VAlign(VAlign_Fill)
+		.HAlign(HAlign_Fill)
+		.Padding(0.f)
+		[
+			SNew(SBox)
+			.WidthOverride(32.f)
+			.HeightOverride(32.f)
+		]
+		+SOverlay::Slot()
+		.VAlign(VAlign_Fill)
+		.HAlign(HAlign_Fill)
+		.Padding(0.f)
+		[
+			SNew(SImage)
+			.ColorAndOpacity(DragAndDropColor)
+			.Image(FDataprepEditorStyle::GetBrush( "DataprepEditor.TrackNode.Slot" ))
+		];
+
+	DummyAction = SNew(SDataprepEmptyActionNode);
 }
 
-void SDataprepGraphTrackWidget::SetActionSlot(int32 SlotIndex, const TSharedRef<SWidget>& Widget)
+void SDataprepGraphTrackWidget::GetDraggedNodePositions(FVector2D& Left, FVector2D& Right, const FGeometry& TargetGeometry)
 {
-	ensure(ActionSlots.IsValidIndex(SlotIndex));
-	(*ActionSlots[SlotIndex])[Widget];
+	const FMargin DragSlotOffset = DragSlot->OffsetAttr.Get(FMargin());
+	const FGeometry& Geometry = GetCachedGeometry();
 
-	Invalidate(EInvalidateWidgetReason::Layout);
+	Left = TargetGeometry.AbsoluteToLocal(Geometry.LocalToAbsolute(FVector2D(DragSlotOffset.Left, 0.f)));
+	Right = TargetGeometry.AbsoluteToLocal(Geometry.LocalToAbsolute(FVector2D(DragSlotOffset.Left + DragSlotOffset.Right, 0.f)));
 }
 
-FVector2D SDataprepGraphTrackWidget::ComputeDesiredSize(float InLayoutScaleMultiplier) const
+void SDataprepGraphTrackWidget::SetActionSlot(int32 SlotIndex, const TSharedRef<SWidget>& Widget)
 {
-	//TrackCanvas->SlatePrepass(InLayoutScaleMultiplier);
-	return SHorizontalBox::ComputeDesiredSize(InLayoutScaleMultiplier);
+	ensure(ActionSlots.IsValidIndex(SlotIndex));
+	ActionSlots[SlotIndex]->AttachWidget(Widget);
+
+	Invalidate(EInvalidateWidgetReason::Layout);
 }
 
 FVector2D SDataprepGraphTrackWidget::GetActioNodeAbscissaRange()
 {
-	FVector2D Range(SlotSpacing);
+	FVector2D Range(InterSlotSpacing, InterSlotSpacing + SDataprepGraphActionNode::DefaultWidth + InterSlotSpacing);
 
 	if(SDataprepGraphTrackNode* TrackNode = TrackNodePtr.Pin().Get())
 	{
-		for(int32 Index = 0; Index < ActionNodes.Num() - 1; ++Index)
-		{
-			Range.Y += GetNodeSize(ActionNodes[Index]).X + SlotSpacing;
-		}
+		Range.Y = ActionSlots[SlotCount-1]->OffsetAttr.Get(FMargin()).Left;
 	}
 
 	return Range;
@@ -824,30 +1058,18 @@ int32 SDataprepGraphTrackWidget::GetHoveredActionNode(float LeftCorner, float Wi
 	const float CenterAbscissa = LeftCorner + Width * 0.5f;
 
 	// Compute origin abscissa of last action node on track
-	for(int32 Index = 0; Index < ActionNodes.Num(); ++Index)
+	for(int32 Index = 0; Index < SlotCount; ++Index)
 	{
-		const FMargin SlotOffset = ActionSlots[Index]->OffsetAttr.Get(FMargin());
-		const float RelativeAbscissa = CenterAbscissa - SlotOffset.Left;
-		if(RelativeAbscissa > -SlotSpacing && RelativeAbscissa <= SlotOffset.Right)
+		const FMargin SlotOffset = DropSlots[Index]->OffsetAttr.Get(FMargin());
+		const float RelativeAbscissa = CenterAbscissa - (SlotOffset.Left + SlotOffset.Right);
+		const float ActionSlotWidth = ActionSlots[Index]->OffsetAttr.Get(FMargin()).Right;
+		if(RelativeAbscissa > -SlotOffset.Right && RelativeAbscissa <= ActionSlotWidth)
 		{
 			return Index;
 		}
 	}
 
-	return bIsCopying ? ActionNodes.Num() : ActionNodes.Num() - 1;
-}
-
-float SDataprepGraphTrackWidget::ValidateNodeAbscissa(float InAbscissa)
-{
-	if(ActionSlots.Num() > 0)
-	{
-		AbscissaRange.X = ActionSlots[0]->OffsetAttr.Get(FVector2D::ZeroVector).Left;
-		AbscissaRange.Y = ActionSlots[ActionSlots.Num() - 1]->OffsetAttr.Get(FVector2D::ZeroVector).Left;
-
-		return InAbscissa < AbscissaRange.X ? AbscissaRange.X : (InAbscissa > AbscissaRange.Y ? AbscissaRange.Y : InAbscissa);
-	}
-
-	return InAbscissa;
+	return SlotCount - 1;
 }
 
 void SDataprepGraphTrackWidget::OnActionsOrderChanged(const TArray<TSharedPtr<SDataprepGraphActionNode>>& InActionNodes)
@@ -859,41 +1081,29 @@ void SDataprepGraphTrackWidget::OnActionsOrderChanged(const TArray<TSharedPtr<SD
 		ActionNodes[Index] = InActionNodes[Index]->AsShared();
 	}
 
-	RefreshLayout(nullptr);
+	RefreshLayout();
 }
 
-void SDataprepGraphTrackWidget::RefreshLayout(SGraphPanel* GraphPanel)
+void SDataprepGraphTrackWidget::RefreshLayout()
 {
-	WorkingArea = FVector2D(0.f, TrackDesiredHeight);
+	const float NodeTopPadding = LineTopPadding + NodeTopOffset;
+
+	LineSize = FVector2D::ZeroVector;
+	WorkingSize = FVector2D(SDataprepGraphActionNode::DefaultWidth, SDataprepGraphActionNode::DefaultHeight);
 	FMargin DropSlotOffset = DropSlots[0]->OffsetAttr.Get(FMargin());
 
-	float MaxWidth = 0.f;
+	float MaxWidth = SDataprepGraphActionNode::DefaultWidth;
 
-	for(int32 Index = 0; Index < ActionNodes.Num(); ++Index)
+	for(int32 Index = 0; Index < SlotCount; ++Index)
 	{
-		DropSlots[Index]->OffsetAttr.Set(DropSlotOffset);
-		DropSlotOffset.Left += SlotSpacing;
-
-		(*DropSlots[Index])[DropFiller.ToSharedRef()];
-
-		TSharedRef<SWidget>& ActionNode = ActionNodes[Index];
-
-		(*ActionSlots[Index])
-		[
-			ActionNode
-		];
-
 		// Update action's proxy node registered to graph panel
-		StaticCastSharedRef<SDataprepGraphActionNode>(ActionNode)->UpdateProxyNode(FVector2D(DropSlotOffset.Left, 0.f) + TrackOffset);
+		StaticCastSharedRef<SDataprepGraphActionNode>(ActionNodes[Index])->UpdateProxyNode(FVector2D(DropSlotOffset.Left + InterSlotSpacing, NodeTopPadding) + TrackOffset);
 
-		const FVector2D Size = GetNodeSize(ActionNode);
+		const FVector2D Size = UpdateActionSlot(DropSlotOffset, Index, ActionNodes[Index]);
 
-		ActionSlots[Index]->OffsetAttr.Set( FMargin(DropSlotOffset.Left, 0.f, Size.X, Size.Y));
-		DropSlotOffset.Left += Size.X;
-
-		if(WorkingArea.Y < Size.Y)
+		if(WorkingSize.Y < Size.Y)
 		{
-			WorkingArea.Y = Size.Y;
+			WorkingSize.Y = Size.Y;
 		}
 
 		if(MaxWidth < Size.X)
@@ -902,179 +1112,121 @@ void SDataprepGraphTrackWidget::RefreshLayout(SGraphPanel* GraphPanel)
 		}
 	}
 
-	DropSlots[ActionNodes.Num()]->OffsetAttr.Set(DropSlotOffset);
-	(*DropSlots[ActionNodes.Num()])[DropFiller.ToSharedRef()];
+	DropSlots[SlotCount]->OffsetAttr.Set(DropSlotOffset);
+	DropSlots[SlotCount]->AttachWidget(SNullWidget::NullWidget);
+	DropSlotOffset.Left += InterSlotSpacing;
 
-	WorkingArea.X = DropSlotOffset.Left + SlotSpacing + MaxWidth/* * 0.5f + 1.f*/;
-	WorkingArea.Y += 30.f;
+	ActionSlots[SlotCount]->OffsetAttr.Set( FMargin(DropSlotOffset.Left, NodeTopPadding, 0., 0.));
+	ActionSlots[SlotCount]->AttachWidget(SNullWidget::NullWidget);
 
-	if(GraphPanel)
-	{
-		float ZoomAmount = GraphPanel->GetZoomAmount();
-		FVector2D PanelSize = GraphPanel->GetCachedGeometry().GetLocalSize() / ZoomAmount;
+	TrackSlots[SlotCount]->OffsetAttr.Set( FMargin(DropSlotOffset.Left + SDataprepGraphActionNode::DefaultWidth * 0.5f - 16.f, LineTopPadding + TrackSlotTopOffset, 32, 32));
+	TrackSlots[SlotCount]->AttachWidget(SNullWidget::NullWidget);
 
-		CanvasArea.X = PanelSize.X > WorkingArea.X ? PanelSize.X : WorkingArea.X;
-		CanvasArea.Y = PanelSize.Y > WorkingArea.Y ? PanelSize.Y : WorkingArea.Y;
-		CanvasArea = CanvasArea + FVector2D(40.f, 40.f / ZoomAmount);
+	WorkingSize.X = DropSlotOffset.Left + SDataprepGraphActionNode::DefaultWidth + InterSlotSpacing;
+	WorkingSize.Y += NodeTopPadding;
 
-		CanvasSlot->OffsetAttr.Set(FMargin(-40.f, -40.f / ZoomAmount, 0.f, 0.f));
-	}
+	UpdateLine(SlotCount-1);
 
 	Invalidate(EInvalidateWidgetReason::Layout | EInvalidateWidgetReason::ChildOrder);
 }
 
-void SDataprepGraphTrackWidget::UpdateLayoutForMove()
+void SDataprepGraphTrackWidget::UpdateLayoutOnDrag(int8 DragDirection)
 {
-	// Clear the leftest slot
-	ActionSlots[ActionNodes.Num()]->OffsetAttr.Set(FVector2D::ZeroVector);
-	(*ActionSlots[ActionNodes.Num()])[SNullWidget::NullWidget];
+	FMargin SlotOffset = DropSlots[0]->OffsetAttr.Get(FMargin());
 
-	// Bring the dragged slot to the last fixed slot if applicable
-	FMargin DragSlotOffset = DragSlot->OffsetAttr.Get(FVector2D::ZeroVector);
-	if(DragSlotOffset.Left > AbscissaRange.Y)
+	if(bIsCopying)
 	{
-		DragSlotOffset.Left = AbscissaRange.Y;
-		DragSlot->OffsetAttr.Set(DragSlotOffset);
-	}
+		for(int32 Index = 0; Index < HoveredDragIndex; ++Index)
+		{
+			UpdateActionSlot(SlotOffset, Index, ActionNodes[Index]);
+		}
+
+		UpdateActionSlot(SlotOffset, HoveredDragIndex, DragSlot->GetWidget(), true);
 
-	// Reconstruct action order if dragged node at left end
-	if(DragIndex == ActionNodes.Num())
+		for(int32 Index = HoveredDragIndex + 1; Index < SlotCount; ++Index)
+		{
+			UpdateActionSlot(SlotOffset, Index, ActionNodes[Index-1]);
+		}
+	}
+	else if(HoveredDragIndex < DragStartIndex)
 	{
-		for(int32 Index = 0; Index < DragStartIndex; ++Index)
+		for(int32 Index = 0; Index < HoveredDragIndex; ++Index)
 		{
-			NewActionsOrder[Index] = Index;
+			UpdateActionSlot(SlotOffset, Index, ActionNodes[Index]);
 		}
 
-		for(int32 Index = DragStartIndex; Index < DragIndex; ++Index)
+		UpdateActionSlot(SlotOffset, HoveredDragIndex, DragSlot->GetWidget(), true);
+
+		for(int32 Index = HoveredDragIndex + 1; Index <= DragStartIndex; ++Index)
 		{
-			NewActionsOrder[Index] = Index + 1;
+			UpdateActionSlot(SlotOffset, Index, ActionNodes[Index-1]);
 		}
 
-		--DragIndex;
-		NewActionsOrder[DragIndex] = DragStartIndex;
+		for(int32 Index = DragStartIndex+1; Index < SlotCount; ++Index)
+		{
+			UpdateActionSlot(SlotOffset, Index, ActionNodes[Index]);
+		}
 	}
-
-	FMargin DropSlotOffset = DropSlots[0]->OffsetAttr.Get(FMargin());
-
-	for(int32 Index = 0; Index < ActionNodes.Num(); ++Index)
+	else
 	{
-		DropSlots[Index]->OffsetAttr.Set(DropSlotOffset);
-		DropSlotOffset.Left += SlotSpacing;
-
-		if(SConstraintCanvas::FSlot* ActionSlot = ActionSlots[Index])
+		for(int32 Index = 0; Index < DragStartIndex; ++Index)
 		{
-			FVector2D Size;
-
-			if(Index == DragIndex)
-			{
-				Size = GetNodeSize(DragSlot->GetWidget());
-
-				(*ActionSlot)
-					[
-						SNew(SColorBlock)
-						.Color( FLinearColor::Transparent )
-					.Size( Size )
-					];
-			}
-			else
-			{
-				TSharedRef<SWidget>& ActionNode = ActionNodes[NewActionsOrder[Index]];
+			UpdateActionSlot(SlotOffset, Index, ActionNodes[Index]);
+		}
 
-				Size = GetNodeSize(ActionNode);
+		for(int32 Index = DragStartIndex; Index < HoveredDragIndex; ++Index)
+		{
+			UpdateActionSlot(SlotOffset, Index, ActionNodes[Index+1]);
+		}
 
-				(*ActionSlot)
-					[
-						ActionNode
-					];
-			}
+		UpdateActionSlot(SlotOffset, HoveredDragIndex, DragSlot->GetWidget(), true);
 
-			ActionSlot->OffsetAttr.Set( FMargin(DropSlotOffset.Left, 0.f, Size.X, Size.Y));
-			DropSlotOffset.Left += Size.X;
-		}
-		else
+		for(int32 Index = HoveredDragIndex + 1; Index < SlotCount; ++Index)
 		{
-			ensure(false);
+			UpdateActionSlot(SlotOffset, Index, ActionNodes[Index]);
 		}
 	}
 
-	DropSlots[ActionNodes.Num()]->OffsetAttr.Set(DropSlotOffset);
-	DropSlotOffset.Left += SlotSpacing;
-
-	ActionSlots[ActionNodes.Num()]->OffsetAttr.Set(DropSlotOffset);
-	(*ActionSlots[ActionNodes.Num()])[SNullWidget::NullWidget];
-}
-
-void SDataprepGraphTrackWidget::UpdateLayoutForCopy()
-{
 	const FMargin DragSlotOffset = DragSlot->OffsetAttr.Get(FMargin());
-	const int32 HoveredIndex = GetHoveredActionNode( DragSlotOffset.Left, DragSlotOffset.Right);
+	const FMargin HoveredDropSlotOffset = DropSlots[HoveredDragIndex]->OffsetAttr.Get(FMargin());
 
-	FMargin SlotOffset = DropSlots[0]->OffsetAttr.Get(FMargin());
-
-	for(int32 Index = 0; Index < HoveredIndex; ++Index)
+	if(HoveredDragIndex > 0 && DragSlotOffset.Left < HoveredDropSlotOffset.Left)
 	{
-		DropSlots[Index]->OffsetAttr.Set(SlotOffset);
-		SlotOffset.Left += SlotSpacing;
-
-		SConstraintCanvas::FSlot* ActionSlot = ActionSlots[Index];
+		SConstraintCanvas::FSlot* OverlappedSlot = ActionSlots[HoveredDragIndex-1];
+		FMargin OverlappedSlotOffset = OverlappedSlot->OffsetAttr.Get(FMargin());
 
-		const FVector2D Size = GetNodeSize(ActionNodes[Index]);
+		const float OverlappedActionSlotLeft = HoveredDropSlotOffset.Left - OverlappedSlotOffset.Right;
+		const float Delta =  HoveredDropSlotOffset.Left - DragSlotOffset.Left;
 
-		ActionSlot->OffsetAttr.Set(FMargin(SlotOffset.Left, 0.f, Size.X, Size.Y));
-		(*ActionSlot)[ActionNodes[Index]];
-
-		SlotOffset.Left += Size.X;
-	}
-
-	if(HoveredIndex == ActionNodes.Num())
-	{
-		DropSlots[ActionNodes.Num()]->OffsetAttr.Set(SlotOffset);
-		SlotOffset.Left += SlotSpacing;
+		const float MaximalLeftOffset = HoveredDropSlotOffset.Left + HoveredDropSlotOffset.Right;
+		const float NewLeftOffset = OverlappedActionSlotLeft + Delta;
 
-		const FVector2D Size = GetNodeSize(DragSlot->GetWidget());
-
-		SConstraintCanvas::FSlot* ActionSlot = ActionSlots[ActionNodes.Num()];
-		ActionSlot->OffsetAttr.Set(FMargin(SlotOffset.Left, 0.f, Size.X, Size.Y));
-		(*ActionSlot)
-		[
-			SNew(SColorBlock)
-			.Color( FLinearColor::Transparent )
-			.Size( Size )
-		];
+		if(NewLeftOffset < MaximalLeftOffset)
+		{
+			OverlappedSlotOffset.Left = NewLeftOffset;
+			OverlappedSlot->OffsetAttr.Set(OverlappedSlotOffset);
+		}
 	}
-	else
+	else if(HoveredDragIndex < (SlotCount-1) && DragSlotOffset.Left > (HoveredDropSlotOffset.Left + HoveredDropSlotOffset.Right) )
 	{
-		DropSlots[HoveredIndex]->OffsetAttr.Set(SlotOffset);
-		SlotOffset.Left += SlotSpacing;
+		SConstraintCanvas::FSlot* OverlappedSlot = ActionSlots[HoveredDragIndex+1];
+		FMargin OverlappedSlotOffset = OverlappedSlot->OffsetAttr.Get(FMargin());
+		FMargin OverlappedDropSlotOffset = DropSlots[HoveredDragIndex+1]->OffsetAttr.Get(FMargin());
 
-		FVector2D Size = GetNodeSize(DragSlot->GetWidget());
+		const float OverlappedActionSlotLeft = OverlappedDropSlotOffset.Left + OverlappedDropSlotOffset.Right;
+		const float MinimalLeftOffset = HoveredDropSlotOffset.Left + HoveredDropSlotOffset.Right;
+		const float Delta = DragSlotOffset.Left - MinimalLeftOffset;
 
-		(*ActionSlots[HoveredIndex])
-		[
-			SNew(SColorBlock)
-			.Color( FLinearColor::Transparent )
-			.Size( Size )
-		];
-
-		ActionSlots[HoveredIndex]->OffsetAttr.Set(FMargin(SlotOffset.Left, 0.f, Size.X, Size.Y));
-		SlotOffset.Left += Size.X;
-
-		for(int32 Index = HoveredIndex + 1; Index <= ActionNodes.Num(); ++Index)
+		float NewLeftOffset = OverlappedActionSlotLeft - Delta;
+		if(NewLeftOffset > MinimalLeftOffset)
 		{
-			DropSlots[Index]->OffsetAttr.Set(SlotOffset);
-			SlotOffset.Left += SlotSpacing;
-
-			SConstraintCanvas::FSlot* ActionSlot = ActionSlots[Index];
-
-			Size = GetNodeSize(ActionNodes[Index - 1]);
-
-			ActionSlot->OffsetAttr.Set(FMargin(SlotOffset.Left, 0.f, Size.X, Size.Y));
-			(*ActionSlot)[ActionNodes[Index - 1]];
-
-			SlotOffset.Left += Size.X;
+			OverlappedSlotOffset.Left = NewLeftOffset;
+			OverlappedSlot->OffsetAttr.Set(OverlappedSlotOffset);
 		}
 	}
 
+	UpdateLine(SlotCount-1);
+
 	Invalidate(EInvalidateWidgetReason::Layout | EInvalidateWidgetReason::Paint);
 }
 
@@ -1089,20 +1241,16 @@ void SDataprepGraphTrackWidget::OnStartNodeDrag(int32 IndexDragged)
 	ensure(TrackNodePtr.IsValid());
 	ensure(ActionSlots.IsValidIndex(IndexDragged));
 
-	NewActionsOrder.SetNum(ActionNodes.Num() + 1);
-	for(int32 Index = 0; Index < NewActionsOrder.Num(); ++Index)
-	{
-		NewActionsOrder[Index] = Index;
-	}
-
 	DragStartIndex = IndexDragged;
 	DragIndex = IndexDragged;
+	HoveredDragIndex = IndexDragged;
+	DragIndicatorIndex = INDEX_NONE;
 
-	FMargin Offset = ActionSlots[IndexDragged]->OffsetAttr.Get(FVector2D::ZeroVector);
+	FMargin Offset = ActionSlots[IndexDragged]->OffsetAttr.Get(FMargin());
 	DragSlot->OffsetAttr.Set(Offset);
-	(*DragSlot)[ActionNodes[IndexDragged]];
+	DragSlot->AttachWidget(ActionNodes[IndexDragged]);
 
-	AbscissaRange.X = ActionSlots[0]->OffsetAttr.Get(FVector2D::ZeroVector).Left;
+	AbscissaRange.X = ActionSlots[0]->OffsetAttr.Get(FMargin()).Left;
 
 	const FModifierKeysState ModifierKeyState = FSlateApplication::Get().GetModifierKeys();
 	bIsCopying = !(ModifierKeyState.IsControlDown() || ModifierKeyState.IsCommandDown());
@@ -1112,13 +1260,15 @@ void SDataprepGraphTrackWidget::OnStartNodeDrag(int32 IndexDragged)
 
 float SDataprepGraphTrackWidget::OnNodeDragged(float DragDelta)
 {
+	float ActualDelta = 0.f;
+
 	if(!DragSlot)
 	{
 		ensure(false);
-		return 0.f;
+		return ActualDelta;
 	}
 
-	FMargin DragSlotOffset = DragSlot->OffsetAttr.Get(FVector2D::ZeroVector);
+	FMargin DragSlotOffset = DragSlot->OffsetAttr.Get(FMargin());
 
 	float NewAbscissa = DragSlotOffset.Left + DragDelta;
 	 
@@ -1132,56 +1282,19 @@ float SDataprepGraphTrackWidget::OnNodeDragged(float DragDelta)
 		const float CachedNodeAbscissa = DragSlotOffset.Left;
 
 		DragSlotOffset.Left = NewAbscissa < AbscissaRange.X ? AbscissaRange.X : (NewAbscissa > AbscissaRange.Y ? AbscissaRange.Y : NewAbscissa);
+		ActualDelta = DragSlotOffset.Left - CachedNodeAbscissa;
 
 		DragSlot->OffsetAttr.Set(DragSlotOffset);
 		DraggedActionNode->GetNodeObj()->NodePosX = DragSlotOffset.Left;
 
 		// Check if center of dragged widget is over a neighboring widget by at least half its size
-		const int32 NewDragIndex = GetHoveredActionNode(DragSlotOffset.Left, DragSlotOffset.Right);
-
-		if(NewDragIndex != DragIndex)
-		{
-			ensure( NewActionsOrder.IsValidIndex(NewDragIndex) );
-
-			// Record the swap
-			int32 TempOrder = NewActionsOrder[DragIndex];
-			NewActionsOrder[DragIndex] = NewActionsOrder[NewDragIndex];
-			NewActionsOrder[NewDragIndex] = TempOrder;
-
-			const int32 CachedIndex = DragIndex;
-			DragIndex = NewDragIndex;
-
-			if(bIsCopying)
-			{
-				UpdateLayoutForCopy();
-			}
-			else
-			{
-				// Cache action node in new index
-				TSharedRef<SWidget> CachedActionNode = ActionSlots[DragIndex]->GetWidget();
-
-				const FVector2D Size = GetNodeSize(DragSlot->GetWidget());
+		HoveredDragIndex = GetHoveredActionNode(DragSlotOffset.Left, DragSlotOffset.Right);
+		UpdateLayoutOnDrag(ActualDelta < 0.f ? -1 : (ActualDelta > 0.f ? 1 : 0));
 
-				(*ActionSlots[DragIndex])
-				[
-					SNew(SColorBlock)
-					.Color( FLinearColor::Transparent )
-					.Size( Size )
-				];
-
-				(*ActionSlots[CachedIndex])
-				[
-					CachedActionNode
-				];
-			}
-		}
-
-		TrackCanvas->Invalidate(EInvalidateWidgetReason::Layout | EInvalidateWidgetReason::ChildOrder);
-
-		return DragSlotOffset.Left - CachedNodeAbscissa;
+		Invalidate(EInvalidateWidgetReason::Layout | EInvalidateWidgetReason::ChildOrder);
 	}
 
-	return 0.f;
+	return ActualDelta;
 }
 
 int32 SDataprepGraphTrackWidget::OnEndNodeDrag()
@@ -1196,64 +1309,84 @@ int32 SDataprepGraphTrackWidget::OnEndNodeDrag()
 	ensure(ActionNode.IsValid());
 
 	bIsCopying = false;
-	int32 CachedLastOrder = DragIndex;
 
-	const FVector2D Size = GetNodeSize(ActionNode->AsShared());
+	SlotCount = ActionNodes.Num();
 
-	DragSlot->OffsetAttr.Set(FVector2D::ZeroVector);
-	(*DragSlot)[SNullWidget::NullWidget];
+	DragSlot->OffsetAttr.Set(FVector2D(-100., -10.f));
+	DragSlot->AttachWidget(SNullWidget::NullWidget);
 
-	DropSlots[ActionNodes.Num()]->OffsetAttr.Set(FVector2D::ZeroVector);
-	(*DropSlots[ActionNodes.Num()])[SNullWidget::NullWidget];
+	DropSlots[SlotCount]->AttachWidget(SNullWidget::NullWidget);
+	ActionSlots[SlotCount]->AttachWidget(SNullWidget::NullWidget);
 
-	RefreshLayout(nullptr);
+	RefreshLayout();
 
+	int32 CachedHoveredDragIndex = HoveredDragIndex;
 	DragStartIndex = INDEX_NONE;
 	DragIndex = INDEX_NONE;
+	HoveredDragIndex = INDEX_NONE;
 
-	NewActionsOrder.Reset();
-
-	TrackCanvas->Invalidate(EInvalidateWidgetReason::Layout | EInvalidateWidgetReason::ChildOrder);
+	Invalidate(EInvalidateWidgetReason::Layout | EInvalidateWidgetReason::ChildOrder);
 
-	return CachedLastOrder;
+	return CachedHoveredDragIndex;
 }
 
 void SDataprepGraphTrackWidget::UpdateDragIndicator(FVector2D MousePosition)
 {
-	ResetDragIndicator();
+	int32 NewDragIndicatorIndex = INDEX_NONE;
 
-	float LocalAbscissa = GetCachedGeometry().AbsoluteToLocal(MousePosition).X - 1.f;
+	const float LocalAbscissa = GetTickSpaceGeometry().AbsoluteToLocal(MousePosition).X - 1.f;
+	FMargin SlotOffset = ActionSlots[0]->OffsetAttr.Get(FMargin());
+	float MaxAbscissa = SlotOffset.Left + (SlotOffset.Right * 0.25f);
 
-	if(LocalAbscissa < DropSlots[0]->OffsetAttr.Get(FMargin()).Left)
+	if(LocalAbscissa < MaxAbscissa)
 	{
-		(*DropSlots[0])[DropIndicator.ToSharedRef()];
-		return;
+		NewDragIndicatorIndex = 0;
 	}
-
-	for(int32 Index = 0; Index < ActionNodes.Num(); ++Index)
+	else
 	{
-		FMargin SlotOffset = ActionSlots[Index]->OffsetAttr.Get(FMargin());
-		if(LocalAbscissa <= (SlotOffset.Left + SlotOffset.Right * 0.5f))
+		for(int32 Index = 0; Index <= SlotCount; ++Index)
 		{
-			(*DropSlots[Index])[DropIndicator.ToSharedRef()];
-			return;
+			SlotOffset = ActionSlots[Index]->OffsetAttr.Get(FMargin());
+			if(LocalAbscissa < (SlotOffset.Left + SlotOffset.Right + 1.f))
+			{
+				NewDragIndicatorIndex = Index;
+				break;
+			}
 		}
-	}
-
-	(*DropSlots[ActionNodes.Num()])[DropIndicator.ToSharedRef()];
-}
 
-int32 SDataprepGraphTrackWidget::GetDragIndicatorIndex()
-{
-	for(int32 Index = 0; Index < DropSlots.Num(); ++Index)
-	{
-		if(DropSlots[Index]->GetWidget() == DropIndicator)
+		if(NewDragIndicatorIndex == INDEX_NONE)
 		{
-			return Index;
+			NewDragIndicatorIndex = SlotCount;
 		}
 	}
 
-	return -1;
+	if(NewDragIndicatorIndex == DragIndicatorIndex)
+	{
+		return;
+	}
+
+	ResetDragIndicator();
+
+	DragIndicatorIndex = NewDragIndicatorIndex;
+
+	SConstraintCanvas::FSlot* TrackSlot = TrackSlots[SlotCount];
+
+	if(NewDragIndicatorIndex == SlotCount)
+	{
+		float TrackLeftOffset = ActionSlots[SlotCount]->OffsetAttr.Get(FMargin()).Left + (SDataprepGraphActionNode::DefaultWidth * 0.5f) - 16.f;
+
+		SlotOffset = TrackSlots[0]->OffsetAttr.Get(FMargin());
+
+		TrackSlot->OffsetAttr.Set(FMargin(TrackLeftOffset, SlotOffset.Top, SlotOffset.Right, SlotOffset.Bottom));
+		TrackSlot->AttachWidget(TrackSlotDrop.ToSharedRef());
+
+		UpdateLine(SlotCount);
+	}
+	else
+	{
+		TrackSlot->OffsetAttr.Set(TrackSlots[DragIndicatorIndex]->OffsetAttr.Get(FMargin()));
+		TrackSlot->AttachWidget(TrackSlotDrop.ToSharedRef());
+	}
 }
 
 float SDataprepGraphTrackWidget::OnControlKeyDown(bool bCtrlDown)
@@ -1263,37 +1396,62 @@ float SDataprepGraphTrackWidget::OnControlKeyDown(bool bCtrlDown)
 		return 0.f;
 	}
 
+	if(DragStartIndex == INDEX_NONE)
+	{
+		bIsCopying = false;
+		return 0.f;
+	}
+
 	bIsCopying = bCtrlDown;
 
 	float MouseOffset = 0.f;
+	SlotCount = ActionNodes.Num();
 
-	if(bIsCopying)
+	FMargin DragSlotOffset = DragSlot->OffsetAttr.Get(FMargin());
+	HoveredDragIndex = GetHoveredActionNode(DragSlotOffset.Left, DragSlotOffset.Right);
+
+	for(int32 Index = 0; Index < SlotCount; ++Index)
 	{
-		UpdateLayoutForCopy();
+		ActionSlots[SlotCount]->AttachWidget(ActionNodes[Index]);
+	}
 
+	if(bIsCopying)
+	{
 		// Set the maximum move on the right
 		const FVector2D DragSize = GetNodeSize(DragSlot->GetWidget());
-		AbscissaRange.Y = ActionSlots[ActionNodes.Num()]->OffsetAttr.Get(FVector2D::ZeroVector).Left + 1.f - (DragSize.X * 0.5f);
+		FMargin LastDropSlotOffset = DropSlots[SlotCount]->OffsetAttr.Get(FMargin());
+		AbscissaRange.Y = LastDropSlotOffset.Left + LastDropSlotOffset.Right + 1.f - (DragSize.X * 0.5f);
+
+		DropSlots[SlotCount]->AttachWidget(DropFiller.ToSharedRef());
+		ActionSlots[SlotCount]->AttachWidget(DummyAction.ToSharedRef());
+		TrackSlots[SlotCount]->AttachWidget(TrackSlotRegular.ToSharedRef());
+
+		++SlotCount;
+		UpdateLayoutOnDrag();
 	}
 	else
 	{
-		AbscissaRange.Y = ActionSlots[ActionNodes.Num() - 1]->OffsetAttr.Get(FVector2D::ZeroVector).Left;
+		FMargin LastDropSlotOffset = DropSlots[SlotCount - 1]->OffsetAttr.Get(FMargin());
+		AbscissaRange.Y = LastDropSlotOffset.Left + LastDropSlotOffset.Right;
 
-		if(DragSlot->OffsetAttr.Get(FVector2D::ZeroVector).Left > AbscissaRange.Y)
+		if(DragSlotOffset.Left > AbscissaRange.Y)
 		{
-			// Update offset of DragSlot
-			FMargin DragSlotOffset = DragSlot->OffsetAttr.Get(FVector2D::ZeroVector);
-
 			MouseOffset = DragSlotOffset.Left - AbscissaRange.Y;
 
 			DragSlotOffset.Left = AbscissaRange.Y;
 			DragSlot->OffsetAttr.Set(DragSlotOffset);
+
+			HoveredDragIndex = GetHoveredActionNode(DragSlotOffset.Left, DragSlotOffset.Right);
 		}
 
-		UpdateLayoutForMove();
+		DropSlots[SlotCount]->AttachWidget(SNullWidget::NullWidget);
+		ActionSlots[SlotCount]->AttachWidget(SNullWidget::NullWidget);
+		TrackSlots[SlotCount]->AttachWidget(SNullWidget::NullWidget);
+
+		UpdateLayoutOnDrag();
 	}
 
-	TrackCanvas->Invalidate(EInvalidateWidgetReason::ChildOrder | EInvalidateWidgetReason::Layout);
+	Invalidate(EInvalidateWidgetReason::ChildOrder | EInvalidateWidgetReason::Layout);
 
 	return MouseOffset;
 }
@@ -1302,7 +1460,7 @@ bool SDataprepGraphTrackWidget::IsDraggedNodeOnEnds()
 {
 	if(DragSlot)
 	{
-		const FMargin DragSlotOffset = DragSlot->OffsetAttr.Get(FVector2D::ZeroVector);
+		const FMargin DragSlotOffset = DragSlot->OffsetAttr.Get(FMargin());
 		return DragSlotOffset.Left == AbscissaRange.X || DragSlotOffset.Left == AbscissaRange.Y;
 	}
 
@@ -1319,11 +1477,28 @@ FVector2D SDataprepGraphTrackWidget::GetNodeSize(const TSharedRef<SWidget>& Widg
 		Size = Widget->GetDesiredSize();
 		if(Size == FVector2D::ZeroVector)
 		{
-			Size.Set(NodeDesiredWidth, NodeDesiredWidth);
+			Size.Set(SDataprepGraphActionNode::DefaultWidth, SDataprepGraphActionNode::DefaultWidth);
 		}
 	}
 
 	return Size;
 }
 
+void SDataprepGraphTrackWidget::UpdateLine(int32 LastIndex)
+{
+	LineSize = FVector2D(0.f, 4.f);
+
+	if(SlotCount > 1)
+	{
+		FMargin OffsetDiff = TrackSlots[LastIndex]->OffsetAttr.Get(FMargin()) - TrackSlots[0]->OffsetAttr.Get(FMargin());
+		FMargin LineOffset = LineSlot->OffsetAttr.Get(FMargin());
+
+		LineOffset.Left = TrackSlots[0]->OffsetAttr.Get(FMargin()).Left + 16.f;
+		LineOffset.Right = OffsetDiff.Left;
+		LineSize.X = OffsetDiff.Left;
+
+		LineSlot->OffsetAttr.Set(LineOffset);
+	}
+}
+
 #undef LOCTEXT_NAMESPACE
\ No newline at end of file
diff --git a/Engine/Plugins/Enterprise/DataprepEditor/Source/DataprepEditor/Private/Widgets/DataprepGraph/SDataprepGraphTrackNode.h b/Engine/Plugins/Enterprise/DataprepEditor/Source/DataprepEditor/Private/Widgets/DataprepGraph/SDataprepGraphTrackNode.h
--- a/Engine/Plugins/Enterprise/DataprepEditor/Source/DataprepEditor/Private/Widgets/DataprepGraph/SDataprepGraphTrackNode.h
+++ b/Engine/Plugins/Enterprise/DataprepEditor/Source/DataprepEditor/Private/Widgets/DataprepGraph/SDataprepGraphTrackNode.h
@@ -13,6 +13,7 @@
 #include "Framework/Commands/UICommandList.h"
 #include "GraphEditor.h"
 #include "GraphEditorActions.h"
+#include "Layout/SlateRect.h"
 #include "SGraphNode.h"
 #include "UObject/StrongObjectPtr.h"
 #include "Widgets/DeclarativeSyntaxSupport.h"
@@ -79,6 +80,7 @@ public:
 	virtual void MoveTo( const FVector2D& NewPosition, FNodeSet& NodeFilter ) override;
 	virtual const FSlateBrush* GetShadowBrush(bool bSelected) const override;
 	virtual FReply OnDrop( const FGeometry& MyGeometry, const FDragDropEvent& DragDropEvent ) override;
+	virtual bool ShouldAllowCulling() const override { return false; }
 	// End of SNodePanel::SNode interface
 
 	// SWidget interface
@@ -89,9 +91,6 @@ public:
 	UDataprepAsset* GetDataprepAsset() { return DataprepAssetPtr.Get(); }
 	const UDataprepAsset* GetDataprepAsset() const { return DataprepAssetPtr.Get(); }
 
-	/** Recompute the boundaries of the graph based on the new size and the new zoom factor */
-	FVector2D Update(const FVector2D& LocalSize, float ZoomAmount);
-
 	void OnControlKeyChanged(bool bControlKeyDown);
 
 	/** Initiates the horizontal drag of an action node */
@@ -112,14 +111,16 @@ public:
 	/** Recomputes the position of each action node */
 	bool RefreshLayout();
 
+	/** Recompute the boundaries of the graph based on the new size and/or the new zoom factor in graph panel */
+	FSlateRect Update();
+
 	/** Start editing of action asset associated to input EdGraphNode */
 	void RequestRename(const UEdGraphNode* Node);
 
 	void RequestViewportPan(const FVector2D& ScreenSpacePosition);
 
 	/** Miscellaneous values used in the display */
-	// #ueent_wip: Will be moved to the Dataprep editor's style
-	static FMargin NodePadding;
+	static FVector2D TrackAnchor;
 
 protected:
 	// SWidget interface
@@ -136,7 +137,7 @@ private:
 	/** Weak pointer to the Dataprep asset holding the displayed actions */
 	TWeakObjectPtr<UDataprepAsset> DataprepAssetPtr;
 
-	/** Range for abscissa of action nodes */
+	/** Range for abscissa of action nodes used during drag-and-drop */
 	FVector2D AbscissaRange;
 
 	/** Indicates a drag is happening */
