Update Dataprep operation "Spawn Actors At Location" to support more asset types

diff --git a/Engine/Plugins/Enterprise/DataprepEditor/Source/DataprepLibraries/Private/DataprepEditingOperations.cpp b/Engine/Plugins/Enterprise/DataprepEditor/Source/DataprepLibraries/Private/DataprepEditingOperations.cpp
--- a/Engine/Plugins/Enterprise/DataprepEditor/Source/DataprepLibraries/Private/DataprepEditingOperations.cpp
+++ b/Engine/Plugins/Enterprise/DataprepEditor/Source/DataprepLibraries/Private/DataprepEditingOperations.cpp
@@ -7,8 +7,11 @@
 #include "IDataprepProgressReporter.h"
 
 #include "ActorEditorUtils.h"
+#include "ActorFactories/ActorFactory.h"
 #include "AssetRegistryModule.h"
 #include "Async/ParallelFor.h"
+#include "IContentBrowserSingleton.h"
+#include "ContentBrowserModule.h"
 #include "Editor.h"
 #include "EditorLevelLibrary.h"
 #include "Engine/StaticMesh.h"
@@ -26,6 +29,13 @@
 #include "ObjectTools.h"
 #include "StaticMeshAttributes.h"
 #include "StaticMeshResources.h"
+#include "ScopedTransaction.h"
+
+// UI related section
+#include "DetailCategoryBuilder.h"
+#include "DetailLayoutBuilder.h"
+#include "DetailWidgetRow.h"
+#include "Widgets/Input/SButton.h"
 
 #define LOCTEXT_NAMESPACE "DatasmithEditingOperations"
 
@@ -614,9 +624,9 @@ void UDataprepSpawnActorsAtLocation::OnExecution_Implementation(const FDataprepC
 	DataprepEditingOperationTime::FTimeLogger TimeLogger(TEXT("SpawnActorsAtLocation"), [&](FText Text) { this->LogInfo(Text); });
 #endif
 
-	if (!StaticMesh)
+	if (!SelectedAsset)
 	{
-		UE_LOG(LogDataprep, Log, TEXT("No mesh was selected"));
+		UE_LOG(LogDataprep, Log, TEXT("No asset was selected"));
 		return;
 	}
 
@@ -629,21 +639,244 @@ void UDataprepSpawnActorsAtLocation::OnExecution_Implementation(const FDataprepC
 
 		if (AActor* Actor = Cast<AActor>(Object))
 		{
-			AStaticMeshActor* NewActor = Cast<AStaticMeshActor>( CreateActor(AStaticMeshActor::StaticClass(), Actor->GetName()) );
+			const FAssetData AssetData(SelectedAsset);
+
+			if (!AssetData.IsValid())
+			{
+				continue;
+			}
+
+			UClass* AssetClass = Cast<UClass>(SelectedAsset);
+
+			if (AssetClass == nullptr)
+			{
+				AssetClass = SelectedAsset->GetClass();
+				check(AssetClass);
+			}
+
+			// Find a factory that can create an actor of desired type.
+			// The factory will not be actually used to create the actor, but to verify it can be 
+			// created and to get it's actual class.
+			const TArray<UActorFactory*>& ActorFactories = GEditor->ActorFactories;
+			UActorFactory* ActorFactory = nullptr;
+			for (int32 FactoryIdx = 0; FactoryIdx < ActorFactories.Num(); FactoryIdx++)
+			{
+				// Check if the actor can be created using this factory
+				FText UnusedErrorMessage;
+				if (ActorFactories[FactoryIdx]->CanCreateActorFrom(AssetData, UnusedErrorMessage))
+				{
+					ActorFactory = ActorFactories[FactoryIdx];
+					break;
+				}
+			}
+
+			if (!ActorFactory)
+			{
+				continue;
+			}
+
+			// Create the actor.
+			AActor* DefaultActor = ActorFactory->GetDefaultActor(AssetData);
+			check(DefaultActor);
+
+			AActor* NewActor = CreateActor(DefaultActor->GetClass(), Actor->GetName());
 			check(NewActor);
 
-			NewActor->GetStaticMeshComponent()->SetStaticMesh(StaticMesh);
+			if (AStaticMeshActor* StaticMeshActor = Cast<AStaticMeshActor>( NewActor ))
+			{
+				// Assign mesh asset to the newly created actor
+				UStaticMesh* StaticMesh = CastChecked<UStaticMesh>(SelectedAsset);
+
+				// Change properties
+				UStaticMeshComponent* StaticMeshComponent = StaticMeshActor->GetStaticMeshComponent();
+				check(StaticMeshComponent);
+
+				StaticMeshComponent->UnregisterComponent();
+
+				StaticMeshComponent->SetStaticMesh(StaticMesh);
+				if (StaticMesh->RenderData)
+				{
+					StaticMeshComponent->StaticMeshDerivedDataKey = StaticMesh->RenderData->DerivedDataKey;
+				}
+
+				// Init Component
+				StaticMeshComponent->RegisterComponent();
+			}
+
 			NewActor->SetActorTransform(Actor->GetActorTransform());
 			NewActor->SetActorLabel(Actor->GetActorLabel() + TEXT("_SA"));
 
-			if (Actor->GetParentActor())
+			AActor* ParentActor = Actor->GetAttachParentActor();
+			if (ParentActor)
 			{
-				NewActor->AttachToActor(Actor->GetParentActor(), FAttachmentTransformRules::KeepRelativeTransform);
+				NewActor->AttachToActor(ParentActor, FAttachmentTransformRules::KeepRelativeTransform);
 			}
 		}
 	}
 }
 
+void UDataprepSpawnActorsAtLocation::SetAsset(UObject* Asset)
+{
+	Modify();
+	SelectedAsset = Asset;
+}
+
+TSharedRef< SWidget > FDataprepSpawnActorsAtLocationDetails::CreateWidget()
+{
+	return SNew(SHorizontalBox)
+		+ SHorizontalBox::Slot()
+		.AutoWidth()
+		.Padding(2, 0)
+		.MaxWidth(100.0f)
+		[
+			SAssignNew(AssetPickerAnchor, SComboButton)
+			.ButtonStyle(FEditorStyle::Get(), "PropertyEditor.AssetComboStyle")
+			.ContentPadding(FMargin(2, 2, 2, 1))
+			.MenuPlacement(MenuPlacement_BelowAnchor)
+			.ButtonContent()
+			[
+				SNew(STextBlock)
+				.TextStyle(FEditorStyle::Get(), "PropertyEditor.AssetClass")
+				.Font(FEditorStyle::GetFontStyle("PropertyWindow.NormalFont"))
+				.Text(this, &FDataprepSpawnActorsAtLocationDetails::OnGetComboTextValue)
+				.ToolTipText(this, &FDataprepSpawnActorsAtLocationDetails::GetObjectToolTip)
+			]
+			.OnGetMenuContent(this, &FDataprepSpawnActorsAtLocationDetails::GenerateAssetPicker)
+		];
+}
+
+const FAssetData& FDataprepSpawnActorsAtLocationDetails::GetAssetData() const
+{
+	if (DataprepOperation->SelectedAsset)
+	{
+		if (!DataprepOperation->SelectedAsset->GetPathName().Equals(CachedAssetData.ObjectPath.ToString(), ESearchCase::CaseSensitive))
+		{
+			// This always uses the exact object pointed at
+			CachedAssetData = FAssetData(DataprepOperation->SelectedAsset, true);
+		}
+	}
+	else
+	{
+		if (CachedAssetData.IsValid())
+		{
+			CachedAssetData = FAssetData();
+		}
+	}
+
+	return CachedAssetData;
+}
+
+FText FDataprepSpawnActorsAtLocationDetails::GetObjectToolTip() const
+{
+	FText Value = FText::GetEmpty();
+	const FAssetData& CurrentAssetData = GetAssetData();
+	
+	if (CurrentAssetData.IsValid())
+	{
+		Value = FText::FromString(CurrentAssetData.GetFullName());
+	}
+	return Value;
+}
+
+FText FDataprepSpawnActorsAtLocationDetails::OnGetComboTextValue() const
+{
+	FText Value = LOCTEXT( "DefaultComboText", "Select Asset" );
+
+	if (DataprepOperation->SelectedAsset != nullptr)
+	{
+		const FAssetData& CurrentAssetData = GetAssetData();
+		if (CurrentAssetData.IsValid())
+		{
+			Value = FText::FromString(CurrentAssetData.AssetName.ToString());
+		}
+	}
+	return Value;
+}
+
+TSharedRef<SWidget> FDataprepSpawnActorsAtLocationDetails::GenerateAssetPicker()
+{
+	FContentBrowserModule& ContentBrowserModule = FModuleManager::Get().LoadModuleChecked<FContentBrowserModule>(TEXT("ContentBrowser"));
+
+	FAssetPickerConfig AssetPickerConfig;
+	AssetPickerConfig.Filter.ClassNames.Add(UStaticMesh::StaticClass()->GetFName());
+	AssetPickerConfig.Filter.ClassNames.Add(UBlueprint::StaticClass()->GetFName());
+	AssetPickerConfig.bAllowNullSelection = true;
+	AssetPickerConfig.Filter.bRecursiveClasses = true;
+	AssetPickerConfig.OnAssetSelected = FOnAssetSelected::CreateSP(this, &FDataprepSpawnActorsAtLocationDetails::OnAssetSelectedFromPicker);
+	AssetPickerConfig.OnAssetEnterPressed = FOnAssetEnterPressed::CreateSP(this, &FDataprepSpawnActorsAtLocationDetails::OnAssetEnterPressedInPicker);
+	AssetPickerConfig.InitialAssetViewType = EAssetViewType::List;
+	AssetPickerConfig.bAllowDragging = false;
+
+	return
+		SNew(SBox)
+		.HeightOverride(300)
+		.WidthOverride(300)
+		[
+			SNew(SBorder)
+			.BorderImage(FEditorStyle::GetBrush("Menu.Background"))
+			[
+				ContentBrowserModule.Get().CreateAssetPicker(AssetPickerConfig)
+			]
+		];
+}
+
+void FDataprepSpawnActorsAtLocationDetails::OnAssetSelectedFromPicker(const struct FAssetData& AssetData)
+{
+	if (AssetPickerAnchor.IsValid())
+	{
+		AssetPickerAnchor->SetIsOpen(false);
+
+		// We need to explicitly delete the customized UI here or a warning will be generated from 
+		// the call to UDataprepSpawnActorsAtLocation::Modify, which re-creates the UI
+		AssetPickerAnchor.Reset();
+	}
+
+	{
+		FScopedTransaction Transaction( LOCTEXT("SpawnActorSetAsset", "Choose Asset for Spawn Actor") );
+		DataprepOperation->SetAsset(AssetData.GetAsset());
+	}
+}
+
+void FDataprepSpawnActorsAtLocationDetails::OnAssetEnterPressedInPicker(const TArray<FAssetData>& InSelectedAssets)
+{
+	if (InSelectedAssets.Num() > 0)
+	{
+		OnAssetSelectedFromPicker(InSelectedAssets[0]);
+	}
+}
+
+void FDataprepSpawnActorsAtLocationDetails::CustomizeDetails(IDetailLayoutBuilder & DetailBuilder)
+{
+	TArray< TWeakObjectPtr< UObject > > Objects;
+
+	DetailBuilder.GetObjectsBeingCustomized(Objects);
+	check(Objects.Num() > 0);
+
+	DataprepOperation = Cast< UDataprepSpawnActorsAtLocation >(Objects[0].Get());
+	check(DataprepOperation);
+
+	DetailBuilder.HideCategory(FName(TEXT("Warning")));
+	IDetailCategoryBuilder& ImportSettingsCategoryBuilder = DetailBuilder.EditCategory(FName(TEXT("SelectedAsset_Internal")), FText::GetEmpty(), ECategoryPriority::Important);
+
+	// Hide SelectedAsset property as it is replaced with custom widget
+	DetailBuilder.HideProperty(GET_MEMBER_NAME_CHECKED(UDataprepSpawnActorsAtLocation, SelectedAsset));
+
+	FDetailWidgetRow& CustomAssetImportRow = ImportSettingsCategoryBuilder.AddCustomRow(FText::FromString(TEXT("Selected Asset")));
+
+	CustomAssetImportRow.NameContent()
+		[
+			SNew(STextBlock)
+			.Text(LOCTEXT("DatasmithActorOperationsLabel", "Selected Asset"))
+			.ToolTipText(LOCTEXT("DatasmithMeshOperationsTooltip", "Selected Asset to spawn Actor from"))
+			.Font(DetailBuilder.GetDetailFont())
+		];
+	
+	CustomAssetImportRow.ValueContent()
+		[
+			CreateWidget()
+		];
+}
+
 bool UDataprepCompactSceneGraphOperation::IsActorVisible(AActor* Actor, TMap<AActor*, bool>& VisibilityMap)
 {
 	if (!Actor)
diff --git a/Engine/Plugins/Enterprise/DataprepEditor/Source/DataprepLibraries/Private/DataprepEditingOperations.h b/Engine/Plugins/Enterprise/DataprepEditor/Source/DataprepLibraries/Private/DataprepEditingOperations.h
--- a/Engine/Plugins/Enterprise/DataprepEditor/Source/DataprepLibraries/Private/DataprepEditingOperations.h
+++ b/Engine/Plugins/Enterprise/DataprepEditor/Source/DataprepLibraries/Private/DataprepEditingOperations.h
@@ -7,6 +7,9 @@
 #include "DataprepOperation.h"
 
 #include "DataprepOperationsLibrary.h"
+#include "AssetData.h"
+
+#include "IDetailCustomization.h"
 
 //
 #include "DataprepEditingOperations.generated.h"
@@ -170,10 +173,42 @@ public:
 		return FDataprepOperationCategories::ActorOperation;
 	}
 
+	void SetAsset(UObject* Asset);
+
 	UPROPERTY(EditAnywhere, BlueprintReadWrite, Category = ActorOperation)
-	UStaticMesh* StaticMesh;
+	UObject* SelectedAsset;
 
 protected:
 	virtual void OnExecution_Implementation(const FDataprepContext& InContext) override;
 	//~ End UDataprepOperation Interface
 };
+
+class FDataprepSpawnActorsAtLocationDetails : public IDetailCustomization
+{
+public:
+	static TSharedRef< IDetailCustomization > MakeDetails() 
+	{ 
+		return MakeShared<FDataprepSpawnActorsAtLocationDetails>(); 
+	};
+
+	FDataprepSpawnActorsAtLocationDetails() : DataprepOperation(nullptr) {}
+	
+	/** Called when details should be customized */
+	virtual void CustomizeDetails(IDetailLayoutBuilder& DetailBuilder) override;
+
+private:
+	const FAssetData& GetAssetData() const;
+	TSharedRef<SWidget> CreateWidget();
+	TSharedRef<SWidget> GenerateAssetPicker();
+	FText OnGetComboTextValue() const;
+	FText GetObjectToolTip() const;
+	void OnAssetSelectedFromPicker(const struct FAssetData& AssetData);
+	void OnAssetEnterPressedInPicker(const TArray<FAssetData>& InSelectedAssets);
+
+private:
+	TSharedPtr<class SMenuAnchor> AssetPickerAnchor;
+	UDataprepSpawnActorsAtLocation* DataprepOperation;
+
+	/** Cached AssetData of object selected */
+	mutable FAssetData CachedAssetData;
+};
diff --git a/Engine/Plugins/Enterprise/DataprepEditor/Source/DataprepLibraries/Private/DataprepLibrariesModule.cpp b/Engine/Plugins/Enterprise/DataprepEditor/Source/DataprepLibraries/Private/DataprepLibrariesModule.cpp
--- a/Engine/Plugins/Enterprise/DataprepEditor/Source/DataprepLibraries/Private/DataprepLibrariesModule.cpp
+++ b/Engine/Plugins/Enterprise/DataprepEditor/Source/DataprepLibraries/Private/DataprepLibrariesModule.cpp
@@ -3,6 +3,7 @@
 #include "DataprepLibrariesModule.h"
 #include "DataprepOperations.h"
 #include "DataprepOperationsLibrary.h"
+#include "DataprepEditingOperations.h"
 
 #include "Modules/ModuleManager.h"
 #include "PropertyEditorModule.h"
@@ -14,6 +15,7 @@ public:
 	{
 		FPropertyEditorModule& PropertyModule = FModuleManager::LoadModuleChecked< FPropertyEditorModule >( TEXT("PropertyEditor") );
 		PropertyModule.RegisterCustomClassLayout( UDataprepSetLODGroupOperation::StaticClass()->GetFName(), FOnGetDetailCustomizationInstance::CreateStatic( &FDataprepSetLOGGroupDetails::MakeDetails ) );
+		PropertyModule.RegisterCustomClassLayout( UDataprepSpawnActorsAtLocation::StaticClass()->GetFName(), FOnGetDetailCustomizationInstance::CreateStatic(&FDataprepSpawnActorsAtLocationDetails::MakeDetails ) );
 		PropertyModule.RegisterCustomClassLayout( UDataprepSetOutputFolder::StaticClass()->GetFName(), FOnGetDetailCustomizationInstance::CreateStatic( &FDataprepSetOutputFolderDetails::MakeDetails ) );
 	}
 
