Fixed UI bleeding out of action widget
Fixed localization warning
Fixed UI update issue when enabling filter preview

#rnx

diff --git a/Engine/Plugins/Enterprise/DataprepEditor/Source/DataprepEditor/Private/SchemaActions/DataprepDragDropOp.cpp b/Engine/Plugins/Enterprise/DataprepEditor/Source/DataprepEditor/Private/SchemaActions/DataprepDragDropOp.cpp
--- a/Engine/Plugins/Enterprise/DataprepEditor/Source/DataprepEditor/Private/SchemaActions/DataprepDragDropOp.cpp
+++ b/Engine/Plugins/Enterprise/DataprepEditor/Source/DataprepEditor/Private/SchemaActions/DataprepDragDropOp.cpp
@@ -320,26 +320,22 @@ FText FDataprepDragDropOp::GetMessageText()
 
 	if(bDropTargetValid || bCopyRequested)
 	{
-		if(Cast<UDataprepGraphActionStepNode>(GetHoveredNode()) != nullptr)
+		if(GetHoveredNode()->IsA<UDataprepGraphActionStepNode>() || GetHoveredNode()->IsA<UDataprepGraphActionNode>())
 		{
-			LastMessageText = bCopyRequested ? LOCTEXT("CopyDataprepActionStepNode", "Copy step to location") : LOCTEXT("MoveDataprepActionStepNode", "Move step to location");
+			LastMessageText = bCopyRequested ? LOCTEXT("DataprepActionStepNode_Copy", "Copy step to location") : LOCTEXT("DataprepActionStepNode_Move", "Move step to location");
 		}
-		else if(Cast<UDataprepGraphActionNode>(GetHoveredNode()) != nullptr)
+		else if(GetHoveredNode()->IsA<UDataprepGraphRecipeNode>())
 		{
-			LastMessageText = bCopyRequested ? LOCTEXT("CopyDataprepActionAssetNode", "Copy step to location") : LOCTEXT("MoveDataprepActionAssetNode", "Move step to location");
-		}
-		else if(Cast<UDataprepGraphRecipeNode>(GetHoveredNode()) != nullptr)
-		{
-			LastMessageText = LOCTEXT("InsertDataprepActionAssetNode", "Insert action to location");
+			LastMessageText = LOCTEXT("DataprepActionAssetNode_Insert", "Insert action to location");
 		}
 	}
 	else if(GetHoveredNode() != nullptr)
 	{
-		LastMessageText = LOCTEXT("NoSelfMove", "Cannot move to itself");
+		LastMessageText = LOCTEXT("DataprepActionStepNode_NoSelfMove", "Cannot move to itself");
 	}
 	else
 	{
-		LastMessageText = LOCTEXT("NotImplemented", "Operation not implemented yet");
+		LastMessageText = LOCTEXT("DataprepActionStepNode_NotImplemented", "Operation not implemented yet");
 	}
 
 	return LastMessageText;
diff --git a/Engine/Plugins/Enterprise/DataprepEditor/Source/DataprepEditor/Private/Widgets/DataprepGraph/SDataprepGraphActionNode.cpp b/Engine/Plugins/Enterprise/DataprepEditor/Source/DataprepEditor/Private/Widgets/DataprepGraph/SDataprepGraphActionNode.cpp
--- a/Engine/Plugins/Enterprise/DataprepEditor/Source/DataprepEditor/Private/Widgets/DataprepGraph/SDataprepGraphActionNode.cpp
+++ b/Engine/Plugins/Enterprise/DataprepEditor/Source/DataprepEditor/Private/Widgets/DataprepGraph/SDataprepGraphActionNode.cpp
@@ -26,7 +26,7 @@
 
 #define LOCTEXT_NAMESPACE "DataprepGraphEditor"
 
-float SDataprepGraphActionNode::DefaultWidth = 300.f;
+float SDataprepGraphActionNode::DefaultWidth = 350.f;
 float SDataprepGraphActionNode::DefaultHeight = 100.f;
 
 class SDataprepGraphActionProxyNode : public SGraphNode
@@ -466,7 +466,7 @@ void SDataprepGraphActionNode::UpdateGraphNode()
 	.VAlign(VAlign_Center)
 	[
 		SNew(SBox)
-		.WidthOverride(300.f)
+		.WidthOverride(DefaultWidth)
 		[
 			SNew(SVerticalBox)
 
@@ -849,7 +849,7 @@ FText SDataprepGraphActionNode::GetBottomWidgetText() const
 		}
 	}
 
-	return LOCTEXT("DataprepEmptyActionStepNoLabel", ".....");
+	return LOCTEXT("DataprepEmptyActionStepNoLabel", "-----");
 }
 
 #undef LOCTEXT_NAMESPACE
\ No newline at end of file
diff --git a/Engine/Plugins/Enterprise/DataprepEditor/Source/DataprepEditor/Private/Widgets/DataprepGraph/SDataprepGraphEditor.cpp b/Engine/Plugins/Enterprise/DataprepEditor/Source/DataprepEditor/Private/Widgets/DataprepGraph/SDataprepGraphEditor.cpp
--- a/Engine/Plugins/Enterprise/DataprepEditor/Source/DataprepEditor/Private/Widgets/DataprepGraph/SDataprepGraphEditor.cpp
+++ b/Engine/Plugins/Enterprise/DataprepEditor/Source/DataprepEditor/Private/Widgets/DataprepGraph/SDataprepGraphEditor.cpp
@@ -133,8 +133,7 @@ void SDataprepGraphEditor::Construct(const FArguments& InArgs, UDataprepAsset* I
 	bMustRearrange = false;
 
 	LastLocalSize = FVector2D::ZeroVector;
-	LastLocation = FVector2D( 0.f, -TopPadding );
-	LastZoomAmount = 1.f;
+	LastZoomAmount = BIG_NUMBER;
 
 	FModifierKeysState ModifierKeyState = FSlateApplication::Get().GetModifierKeys();
 	bCachedControlKeyDown = ModifierKeyState.IsControlDown() || ModifierKeyState.IsCommandDown();
@@ -142,18 +141,25 @@ void SDataprepGraphEditor::Construct(const FArguments& InArgs, UDataprepAsset* I
 	SetNodeFactory( MakeShared<FDataprepGraphNodeFactory>( InArgs._DataprepEditor ) );
 }
 
+void SDataprepGraphEditor::NotifyGraphChanged()
+{
+	// Release track node widget as it is about to be re-created
+	TrackGraphNodePtr.Reset();
+	bIsComplete = false;
+
+	// Reset cached size and zoom. No need for location because it is invalidated when size has changed
+	LastLocalSize = FVector2D::ZeroVector;
+	LastZoomAmount = BIG_NUMBER;
+
+	SGraphEditor::NotifyGraphChanged();
+}
+
 // #ueent_toremove: Temp code for the nodes development
 void SDataprepGraphEditor::OnPipelineChanged(UBlueprint* InBlueprint)
 {
 	if (InBlueprint)
 	{
-		TrackGraphNodePtr.Reset();
-		bIsComplete = false;
 		NotifyGraphChanged();
-
-		LastLocalSize = FVector2D::ZeroVector;
-		//LastLocation = FVector2D( BIG_NUMBER );
-		LastZoomAmount = 1.f;
 	}
 }
 
@@ -179,13 +185,8 @@ void SDataprepGraphEditor::OnDataprepAssetActionChanged(UObject* InObject, FData
 				Nodes.Remove(NodeObject);
 			}
 
-			TrackGraphNodePtr.Reset();
-			bIsComplete = false;
 			NotifyGraphChanged();
 
-			LastLocalSize = FVector2D::ZeroVector;
-			LastLocation = FVector2D::ZeroVector;
-			LastZoomAmount = 1.f;
 			break;
 		}
 
diff --git a/Engine/Plugins/Enterprise/DataprepEditor/Source/DataprepEditor/Private/Widgets/DataprepGraph/SDataprepGraphEditor.h b/Engine/Plugins/Enterprise/DataprepEditor/Source/DataprepEditor/Private/Widgets/DataprepGraph/SDataprepGraphEditor.h
--- a/Engine/Plugins/Enterprise/DataprepEditor/Source/DataprepEditor/Private/Widgets/DataprepGraph/SDataprepGraphEditor.h
+++ b/Engine/Plugins/Enterprise/DataprepEditor/Source/DataprepEditor/Private/Widgets/DataprepGraph/SDataprepGraphEditor.h
@@ -84,6 +84,10 @@ public:
 	virtual FReply OnDrop(const FGeometry& MyGeometry, const FDragDropEvent& DragDropEvent) override;
 	// End of SWidget overrides
 
+	// SGraphEditor overrides
+	virtual void NotifyGraphChanged() override;
+	// End of SGraphEditor overrides
+
 	// #ueent_toremove: Temp code for the nodes development
 	void OnPipelineChanged(UBlueprint* InBlueprint);
 
diff --git a/Engine/Plugins/Enterprise/DataprepEditor/Source/DataprepEditor/Private/Widgets/DataprepWidgets.cpp b/Engine/Plugins/Enterprise/DataprepEditor/Source/DataprepEditor/Private/Widgets/DataprepWidgets.cpp
--- a/Engine/Plugins/Enterprise/DataprepEditor/Source/DataprepEditor/Private/Widgets/DataprepWidgets.cpp
+++ b/Engine/Plugins/Enterprise/DataprepEditor/Source/DataprepEditor/Private/Widgets/DataprepWidgets.cpp
@@ -123,7 +123,7 @@ namespace DataprepWidgetUtils
 
 	private:
 		TSharedPtr< FDataprepDetailsViewColumnSizeData > ColumnSizeData;
-};
+	};
 
 	TSharedRef<SWidget> CreatePropertyWidget( TSharedPtr<SWidget> NameWidget, TSharedPtr<SWidget> ValueWidget, TSharedPtr< FDataprepDetailsViewColumnSizeData > ColumnSizeData, float Spacing)
 	{
@@ -629,7 +629,7 @@ void SDataprepDetailsView::Construct()
 		[
 			SNew(SVerticalBox)
 			+ SVerticalBox::Slot()
-			.Padding(5.0f)
+			.Padding(1.f, 5.0f)
 			[
 				GridPanel.ToSharedRef()
 			]
diff --git a/Engine/Plugins/Enterprise/DataprepEditor/Source/DataprepEditor/Private/Widgets/DataprepWidgets.h b/Engine/Plugins/Enterprise/DataprepEditor/Source/DataprepEditor/Private/Widgets/DataprepWidgets.h
--- a/Engine/Plugins/Enterprise/DataprepEditor/Source/DataprepEditor/Private/Widgets/DataprepWidgets.h
+++ b/Engine/Plugins/Enterprise/DataprepEditor/Source/DataprepEditor/Private/Widgets/DataprepWidgets.h
@@ -164,7 +164,7 @@ private:
 
 	float OnGetLeftColumnWidth() const { return 1.0f - ColumnWidth; }
 	float OnGetRightColumnWidth() const { return ColumnWidth; }
-	void OnSetColumnWidth(float InWidth) { ColumnWidth = InWidth < 0.3f ? 0.3f : InWidth; }
+	void OnSetColumnWidth(float InWidth) { ColumnWidth = InWidth < 0.5f ? 0.5f : InWidth; }
 
 	/** Callback to track property changes on array properties */
 	void OnPropertyChanged( const struct FPropertyChangedEvent& InEvent );
@@ -276,7 +276,7 @@ private:
 
 	float OnGetLeftColumnWidth() const { return 1.0f - ColumnWidth; }
 	float OnGetRightColumnWidth() const { return ColumnWidth; }
-	void OnSetColumnWidth(float InWidth) { ColumnWidth = InWidth; }
+	void OnSetColumnWidth(float InWidth)  { ColumnWidth = InWidth < 0.5f ? 0.5f : InWidth; }
 
 private:
 	/** Weak pointer on DataprepAsset instance */
