Incorrect location of spawned actors using Spawn Actors At Location

diff --git a/Engine/Plugins/Enterprise/DataprepEditor/Source/DataprepLibraries/Private/DataprepEditingOperations.cpp b/Engine/Plugins/Enterprise/DataprepEditor/Source/DataprepLibraries/Private/DataprepEditingOperations.cpp
--- a/Engine/Plugins/Enterprise/DataprepEditor/Source/DataprepLibraries/Private/DataprepEditingOperations.cpp
+++ b/Engine/Plugins/Enterprise/DataprepEditor/Source/DataprepLibraries/Private/DataprepEditingOperations.cpp
@@ -703,7 +703,7 @@ void UDataprepSpawnActorsAtLocation::OnExecution_Implementation(const FDataprepC
 				StaticMeshComponent->RegisterComponent();
 			}
 
-			NewActor->SetActorTransform(Actor->GetActorTransform());
+			NewActor->SetActorRelativeTransform(Actor->GetRootComponent()->GetRelativeTransform());
 			NewActor->SetActorLabel(Actor->GetActorLabel() + TEXT("_SA"));
 
 			AActor* ParentActor = Actor->GetAttachParentActor();
