Fixed an issue where we were duplicating a cdo in the dataprep asset. This was causing some issues with the blueprint reinstancing.

diff --git a/Engine/Plugins/Enterprise/DataprepEditor/Source/DataprepCore/Private/DataprepAsset.cpp b/Engine/Plugins/Enterprise/DataprepEditor/Source/DataprepCore/Private/DataprepAsset.cpp
--- a/Engine/Plugins/Enterprise/DataprepEditor/Source/DataprepCore/Private/DataprepAsset.cpp
+++ b/Engine/Plugins/Enterprise/DataprepEditor/Source/DataprepCore/Private/DataprepAsset.cpp
@@ -623,48 +623,24 @@ int32 UDataprepAsset::AddActions(const TArray<const UDataprepActionAsset*>& InAc
 	return INDEX_NONE;
 }
 
-int32 UDataprepAsset::AddActions(const TArray<const UDataprepActionStep*>& InActionSteps, bool bCreateOne)
+int32 UDataprepAsset::AddAction(const TArray<const UDataprepActionStep*>& InActionSteps)
 {
 	if ( InActionSteps.Num() > 0 && InActionSteps[0] != nullptr )
 	{
 		Modify();
 
-		int32 PreviousActionCount = ActionAssets.Num();
-
-		if(bCreateOne == true)
-		{
-			UDataprepActionAsset* Action = NewObject<UDataprepActionAsset>( this, UDataprepActionAsset::StaticClass(), NAME_None, RF_Transactional );
-			Action->SetLabel( TEXT("New Action") );
+		UDataprepActionAsset* Action = NewObject<UDataprepActionAsset>(this, UDataprepActionAsset::StaticClass(), NAME_None, RF_Transactional);
+		Action->SetLabel(TEXT("New Action"));
 
-			ActionAssets.Add( Action );
+		ActionAssets.Add(Action);
 
-			Action->AddSteps(InActionSteps);
-		}
-		else
-		{
-
-			for(const UDataprepActionStep* InActionStep : InActionSteps)
-			{
-				if(InActionStep)
-				{
-					UDataprepActionAsset* Action = NewObject<UDataprepActionAsset>( this, UDataprepActionAsset::StaticClass(), NAME_None, RF_Transactional );
-					Action->SetLabel( TEXT("New Action") );
-
-					ActionAssets.Add( Action );
-
-					Action->AddStep(InActionStep);
-				}
-			}
-		}
+		Action->AddSteps(InActionSteps);
 
 		CachedActionCount = ActionAssets.Num();
 
-		if(PreviousActionCount != CachedActionCount)
-		{
-			OnActionChanged.Broadcast(ActionAssets.Last(), FDataprepAssetChangeType::ActionAdded);
+		OnActionChanged.Broadcast(Action, FDataprepAssetChangeType::ActionAdded);
 
-			return ActionAssets.Num() - 1;
-		}
+		return ActionAssets.Num() - 1;
 	}
 
 	UE_LOG( LogDataprepCore, Error, TEXT("UDataprepAsset::AddActionSteps: None of the action steps is invalid") );
@@ -676,7 +652,7 @@ int32 UDataprepAsset::AddActions(const TArray<const UDataprepActionStep*>& InAct
 
 bool UDataprepAsset::InsertAction(const UDataprepActionAsset* InAction, int32 Index)
 {
-	if(!ActionAssets.IsValidIndex(Index))
+	if( !( Index >= 0 && Index <= ActionAssets.Num() ) )
 	{
 		UE_LOG( LogDataprepCore, Error, TEXT("UDataprepAsset::InsertAction: The index is invalid") );
 		return false;
@@ -686,13 +662,13 @@ bool UDataprepAsset::InsertAction(const UDataprepActionAsset* InAction, int32 In
 	{
 		Modify();
 
-		UDataprepActionAsset* Action = DuplicateObject<UDataprepActionAsset>( InAction, this);
-		Action->SetFlags(EObjectFlags::RF_Transactional);
-		Action->SetLabel( InAction->GetLabel());
+		UDataprepActionAsset* Action = DuplicateObject<UDataprepActionAsset>( InAction, this );
+		Action->SetFlags( EObjectFlags::RF_Transactional );
+		Action->SetLabel( InAction->GetLabel() );
 
 		ActionAssets.Insert( Action, Index );
 
-		OnActionChanged.Broadcast(Action, FDataprepAssetChangeType::ActionAdded);
+		OnActionChanged.Broadcast( Action, FDataprepAssetChangeType::ActionAdded );
 
 		CachedActionCount = ActionAssets.Num();
 
@@ -740,66 +716,65 @@ bool UDataprepAsset::InsertActions(const TArray<const UDataprepActionAsset*>& In
 
 		if(PreviousActionCount != CachedActionCount)
 		{
+			// We should rework this notification
 			OnActionChanged.Broadcast(ActionAssets.Last(), FDataprepAssetChangeType::ActionAdded);
 
 			return true;
 		}
 	}
 
-	UE_LOG( LogDataprepCore, Error, TEXT("UDataprepAsset::InsertActions: None of the actions is invalid") );
+	UE_LOG( LogDataprepCore, Error, TEXT("UDataprepAsset::InsertActions: one of the actions is invalid") );
 	ensure(false);
 
 	// Invalid
 	return false;
 }
 
-bool UDataprepAsset::InsertActions(const TArray<const UDataprepActionStep*>& InActionSteps, int32 Index, bool bCreateOne)
+bool UDataprepAsset::InsertAction(int32 Index)
 {
-	if ( InActionSteps.Num() > 0 && InActionSteps[0] != nullptr )
+	if ( Index >= 0 && Index <= ActionAssets.Num() )
 	{
-		Modify();
+		UDataprepActionAsset* Action = NewObject<UDataprepActionAsset>( this, UDataprepActionAsset::StaticClass(), NAME_None, RF_Transactional );
+		Action->SetLabel( TEXT("New Action") );
 
-		int32 PreviousActionCount = ActionAssets.Num();
+		ActionAssets.Insert( Action, Index );
 
-		if(bCreateOne == true)
-		{
-			UDataprepActionAsset* Action = NewObject<UDataprepActionAsset>( this, UDataprepActionAsset::StaticClass(), NAME_None, RF_Transactional );
-			Action->SetLabel( TEXT("New Action") );
+		CachedActionCount = ActionAssets.Num();
 
-			ActionAssets.Insert( Action, Index );
+		OnActionChanged.Broadcast( Action, FDataprepAssetChangeType::ActionAdded );
+		return true;
+	}
 
-			Action->AddSteps(InActionSteps);
-		}
-		else
+	UE_LOG( LogDataprepCore, Error, TEXT("UDataprepAsset::AddAction: The index is invalid") );
+	ensure(false);
+	return false;
+}
+
+bool UDataprepAsset::InsertAction(const TArray<const UDataprepActionStep*>& InActionSteps, int32 Index)
+{
+	if ( InActionSteps.Num() > 0 && InActionSteps[0] != nullptr )
+	{
+		if( Index >= 0 && Index <= ActionAssets.Num() )
 		{
-			int32 InsertIndex = Index;
+			Modify();
 
-			for(const UDataprepActionStep* InActionStep : InActionSteps)
-			{
-				if(InActionStep)
-				{
-					UDataprepActionAsset* Action = NewObject<UDataprepActionAsset>( this, UDataprepActionAsset::StaticClass(), NAME_None, RF_Transactional );
-					Action->SetLabel( TEXT("New Action") );
+			int32 PreviousActionCount = ActionAssets.Num();
 
-					ActionAssets.Insert( Action, InsertIndex );
-					++InsertIndex;
+			UDataprepActionAsset* Action = NewObject<UDataprepActionAsset>( this, UDataprepActionAsset::StaticClass(), NAME_None, RF_Transactional );
+			Action->SetLabel( TEXT("New Action") );
 
-					Action->AddStep(InActionStep);
-				}
-			}
-		}
+			ActionAssets.Insert( Action, Index );
 
-		CachedActionCount = ActionAssets.Num();
+			Action->AddSteps( InActionSteps );
 
-		if(CachedActionCount != PreviousActionCount)
-		{
-			OnActionChanged.Broadcast(ActionAssets.Last(), FDataprepAssetChangeType::ActionAdded);
+			CachedActionCount = ActionAssets.Num();
 
+			OnActionChanged.Broadcast( Action, FDataprepAssetChangeType::ActionAdded );
 			return true;
 		}
 	}
 
-	UE_LOG( LogDataprepCore, Error, TEXT("UDataprepAsset::AddAction: None of the action steps is invalid") );
+	UE_LOG( LogDataprepCore, Error, TEXT("UDataprepAsset::InsertAction: One of the action steps is invalid or the index is invalid") );
 	ensure(false);
 
 	// Invalid
diff --git a/Engine/Plugins/Enterprise/DataprepEditor/Source/DataprepCore/Public/DataprepAsset.h b/Engine/Plugins/Enterprise/DataprepEditor/Source/DataprepCore/Public/DataprepAsset.h
--- a/Engine/Plugins/Enterprise/DataprepEditor/Source/DataprepCore/Public/DataprepAsset.h
+++ b/Engine/Plugins/Enterprise/DataprepEditor/Source/DataprepCore/Public/DataprepAsset.h
@@ -77,6 +77,15 @@ public:
 	 */
 	int32 AddAction(const UDataprepActionAsset* Action = nullptr);
 
+	/**
+	 * Creates an action from the array of action steps or one action per action steps
+	 * then add the action(s) to the Dataprep asset
+	 * @param ActionSteps The array of action steps to process
+	 * @param bCreateOne Indicates if one or more action assets should be created. By default one is created
+	 * @return The index of the last added action or index none if the action is invalid
+	 */
+	int32 AddAction(const TArray<const UDataprepActionStep*>& ActionSteps);
+
 	/**
 	 * Add the actions to the Dataprep asset
 	 * @param Actions The array of action to add
@@ -85,14 +94,22 @@ public:
 	 */
 	int32 AddActions(const TArray<const UDataprepActionAsset*>& Actions);
 
+	/**
+	 * Creates a new action 
+	 * then insert the action to the Dataprep asset at the requested index
+	 * @param Index The index at which the insertion must happen
+	 * @return True if the insertion is successful, false if the index is invalid
+	 */
+	bool InsertAction(int32 Index);
+	
 	/**
 	 * Creates an action from the array of action steps or one action per action steps
-	 * then add the action(s) to the Dataprep asset
+	 * then insert the action(s) to the Dataprep asset at the requested index
 	 * @param ActionSteps The array of action steps to process
-	 * @param bCreateOne Indicates if one or more action assets should be created. By default one is created
-	 * @return The index of the last added action or index none if the action is invalid
+	 * @param Index The index at which the insertion must happen
+	 * @return True if the insertion is successful, false if the action steps or the index are invalid
 	 */
-	int32 AddActions(const TArray<const UDataprepActionStep*>& ActionSteps, bool bCreateOne = true);
+	bool InsertAction(const TArray<const UDataprepActionStep*>& InActionSteps, int32 Index);
 
 	/**
 	 * Insert a copy of the action to the Dataprep asset at the requested index
@@ -110,15 +127,6 @@ public:
 	 */
 	bool InsertActions(const TArray<const UDataprepActionAsset*>& InActions, int32 Index);
 
-	/**
-	 * Creates an action from the array of action steps or one action per action steps
-	 * then insert the action(s) to the Dataprep asset at the requested index
-	 * @param ActionSteps The array of action steps to process
-	 * @param Index The index at which the insertion must happen
-	 * @return True if the insertion is successful, false if the action steps or the index are invalid
-	 */
-	bool InsertActions(const TArray<const UDataprepActionStep*>& InActionSteps, int32 Index, bool bCreateOne = true);
-
 	/**
 	 * Move an action to another spot in the order of actions
 	 * This operation take O(n) time. Where n is the absolute value of SourceIndex - DestinationIndex
diff --git a/Engine/Plugins/Enterprise/DataprepEditor/Source/DataprepEditor/Private/SchemaActions/DataprepDragDropOp.cpp b/Engine/Plugins/Enterprise/DataprepEditor/Source/DataprepEditor/Private/SchemaActions/DataprepDragDropOp.cpp
--- a/Engine/Plugins/Enterprise/DataprepEditor/Source/DataprepEditor/Private/SchemaActions/DataprepDragDropOp.cpp
+++ b/Engine/Plugins/Enterprise/DataprepEditor/Source/DataprepEditor/Private/SchemaActions/DataprepDragDropOp.cpp
@@ -548,11 +548,11 @@ FReply FDataprepDragDropOp::DoDropOnTrack(UDataprepAsset* TargetDataprepAsset, i
 
 			if(InsertIndex >= 0 && InsertIndex < TargetDataprepAsset->GetActionCount())
 			{
-				TargetDataprepAsset->InsertActions(Steps, InsertIndex);
+				TargetDataprepAsset->InsertAction(Steps, InsertIndex);
 			}
 			else
 			{
-				TargetDataprepAsset->AddActions(Steps);
+				TargetDataprepAsset->AddAction(Steps);
 			}
 		}
 
@@ -565,15 +565,18 @@ FReply FDataprepDragDropOp::DoDropOnTrack(UDataprepAsset* TargetDataprepAsset, i
 
 		bool bTransactionSuccessful = false;
 
-		UDataprepActionAsset* Action = Cast<UDataprepActionAsset>(UDataprepActionAsset::StaticClass()->GetDefaultObject());
+		UDataprepActionAsset* Action = nullptr;
 
 		if(InsertIndex >= 0 && InsertIndex < TargetDataprepAsset->GetActionCount())
 		{
-			bTransactionSuccessful = TargetDataprepAsset->InsertAction(Action, InsertIndex);
+			bTransactionSuccessful = TargetDataprepAsset->InsertAction(InsertIndex);
+			Action = TargetDataprepAsset->GetAction(InsertIndex);
 		}
 		else
 		{
-			InsertIndex = TargetDataprepAsset->AddAction(Action);
+			InsertIndex = TargetDataprepAsset->AddAction();
+			Action = TargetDataprepAsset->GetAction(InsertIndex);
+
 			bTransactionSuccessful = InsertIndex != INDEX_NONE;
 		}
 
@@ -583,9 +586,6 @@ FReply FDataprepDragDropOp::DoDropOnTrack(UDataprepAsset* TargetDataprepAsset, i
 			return FReply::Unhandled();
 		}
 
-		Action = TargetDataprepAsset->GetAction(InsertIndex);
-		Action->SetLabel(TEXT("New Action"));
-
 		FDataprepSchemaActionContext Context;
 		Context.DataprepActionPtr = Action;
 
