Ensure that the RF_DefaultSubobject flag persists when copying an object in the level editor

diff --git a/Engine/Plugins/Enterprise/DataprepEditor/Source/DataprepEditor/Private/DataprepSnapshot.cpp b/Engine/Plugins/Enterprise/DataprepEditor/Source/DataprepEditor/Private/DataprepSnapshot.cpp
--- a/Engine/Plugins/Enterprise/DataprepEditor/Source/DataprepEditor/Private/DataprepSnapshot.cpp
+++ b/Engine/Plugins/Enterprise/DataprepEditor/Source/DataprepEditor/Private/DataprepSnapshot.cpp
@@ -635,6 +635,12 @@ void FDataprepEditor::TakeSnapshot()
 									bLocalIsValid = FFileHelper::SaveArrayToFile( SerializedData, *AssetFilePath );
 								}
 
+								// Ensure DefaultSubojbect flag persists through the clearing of flags
+								if (AssetObject->HasAllFlags(RF_DefaultSubObject))
+								{
+									ObjectFlags |= RF_DefaultSubObject;
+								}
+
 								AssetObject->ClearFlags( RF_AllFlags );
 								AssetObject->SetFlags( ObjectFlags );
 
diff --git a/Engine/Source/Editor/UnrealEd/Private/EditorObject.cpp b/Engine/Source/Editor/UnrealEd/Private/EditorObject.cpp
--- a/Engine/Source/Editor/UnrealEd/Private/EditorObject.cpp
+++ b/Engine/Source/Editor/UnrealEd/Private/EditorObject.cpp
@@ -573,6 +573,12 @@ static const TCHAR* ImportProperties(
 							NewFlags &= ~RF_Transactional;
 						}
 					}
+					
+					// Ensure DefaultSubojbect flag persists through the clearing of flags
+					if (ComponentTemplate->HasAllFlags(RF_DefaultSubObject))
+					{
+						NewFlags |= RF_DefaultSubObject;
+					}
 
 					// Make sure desired flags are set - existing object could be pending kill
 					ComponentTemplate->ClearFlags(RF_AllFlags);
