Fixed minor but annoying user interaction  issues

#rnx

diff --git a/Engine/Plugins/Enterprise/DataprepEditor/Source/DataprepEditor/Private/DataprepEditor.cpp b/Engine/Plugins/Enterprise/DataprepEditor/Source/DataprepEditor/Private/DataprepEditor.cpp
--- a/Engine/Plugins/Enterprise/DataprepEditor/Source/DataprepEditor/Private/DataprepEditor.cpp
+++ b/Engine/Plugins/Enterprise/DataprepEditor/Source/DataprepEditor/Private/DataprepEditor.cpp
@@ -768,7 +768,7 @@ void FDataprepEditor::OnCommitWorld()
 		const FText Title( LOCTEXT( "DataprepEditor_ProceedWithCommit", "Proceed with commit" ) );
 		const FText Message( LOCTEXT( "DataprepEditor_ConfirmCommitPipelineNotExecuted", "The action pipeline has not been executed.\nDo you want to proceeed with the commit anyway?" ) );
 
-		if( FMessageDialog::Open( EAppMsgType::YesNo, Message, &Title ) == EAppReturnType::No )
+		if( FMessageDialog::Open( EAppMsgType::YesNo, Message, &Title ) != EAppReturnType::Yes )
 		{
 			return;
 		}
@@ -779,7 +779,7 @@ void FDataprepEditor::OnCommitWorld()
 		const FText Title( LOCTEXT( "DataprepEditor_ProceedWithCommit", "Proceed with commit" ) );
 		const FText Message( LOCTEXT( "DataprepEditor_ConfirmCommitPipelineChanged", "The action pipeline has changed since last execution.\nDo you want to proceeed with the commit anyway?" ) );
 
-		if( FMessageDialog::Open( EAppMsgType::YesNo, Message, &Title ) == EAppReturnType::No )
+		if( FMessageDialog::Open( EAppMsgType::YesNo, Message, &Title ) != EAppReturnType::Yes )
 		{
 			return;
 		}
diff --git a/Engine/Plugins/Enterprise/DataprepEditor/Source/DataprepEditor/Private/Widgets/DataprepGraph/SDataprepGraphTrackNode.cpp b/Engine/Plugins/Enterprise/DataprepEditor/Source/DataprepEditor/Private/Widgets/DataprepGraph/SDataprepGraphTrackNode.cpp
--- a/Engine/Plugins/Enterprise/DataprepEditor/Source/DataprepEditor/Private/Widgets/DataprepGraph/SDataprepGraphTrackNode.cpp
+++ b/Engine/Plugins/Enterprise/DataprepEditor/Source/DataprepEditor/Private/Widgets/DataprepGraph/SDataprepGraphTrackNode.cpp
@@ -473,23 +473,38 @@ FReply SDataprepGraphTrackNode::OnDrop(const FGeometry& MyGeometry, const FDragD
 	return SGraphNode::OnDrop(MyGeometry, DragDropEvent);
 }
 
-FReply SDataprepGraphTrackNode::OnDragOver(const FGeometry & MyGeometry, const FDragDropEvent & DragDropEvent)
+FReply SDataprepGraphTrackNode::OnDragOver(const FGeometry & MyGeometry, const FDragDropEvent& DragDropEvent)
 {
 	TSharedPtr<FDataprepDragDropOp> DragActionNodeOp = DragDropEvent.GetOperationAs<FDataprepDragDropOp>();
 	if(DragActionNodeOp.IsValid())
 	{
-		DragActionNodeOp->SetHoveredNode(GraphNode);
-		TrackWidgetPtr->UpdateDragIndicator(DragDropEvent.GetScreenSpacePosition());
+		// It looks like there is a bug in the handling of DnD, OnDragOver is still called after leaving the track.
+		// This seems to be related to the header of the graph panel. The event is going through
+		// Quick fix: Detect onDragOver is called while cursor is not on top of track and manually leave.
+		// #ueent_todo: Fix this at the root.
+		if(!GetTickSpaceGeometry().IsUnderLocation(DragDropEvent.GetScreenSpacePosition()))
+		{
+			if(TrackWidgetPtr->DragIndicatorIndex != INDEX_NONE)
+			{
+				OnDragLeave(DragDropEvent);
+			}
+		}
+		else
+		{
+			DragActionNodeOp->SetHoveredNode(GraphNode);
+			TrackWidgetPtr->UpdateDragIndicator(DragDropEvent.GetScreenSpacePosition());
+		}
 	}
 
 	return SGraphNode::OnDragOver(MyGeometry, DragDropEvent);
 }
 
-void SDataprepGraphTrackNode::OnDragLeave(const FDragDropEvent & DragDropEvent)
+void SDataprepGraphTrackNode::OnDragLeave(const FDragDropEvent& DragDropEvent)
 {
 	TSharedPtr<FDataprepDragDropOp> DragActionNodeOp = DragDropEvent.GetOperationAs<FDataprepDragDropOp>();
 	if(DragActionNodeOp.IsValid())
 	{
+		DragActionNodeOp->SetHoveredNode(nullptr);
 		TrackWidgetPtr->ResetDragIndicator();
 		TrackWidgetPtr->RefreshLayout();
 	}
