Removed all source code related to BP from Dataprep core and editor code

#rnx

diff --git a/Engine/Plugins/Enterprise/DataprepEditor/Source/DataprepCore/Private/DataprepAsset.cpp b/Engine/Plugins/Enterprise/DataprepEditor/Source/DataprepCore/Private/DataprepAsset.cpp
--- a/Engine/Plugins/Enterprise/DataprepEditor/Source/DataprepCore/Private/DataprepAsset.cpp
+++ b/Engine/Plugins/Enterprise/DataprepEditor/Source/DataprepCore/Private/DataprepAsset.cpp
@@ -11,7 +11,6 @@
 #include "DataprepCorePrivateUtils.h"
 #include "DataprepCoreUtils.h"
 #include "DataprepParameterizableObject.h"
-#include "DataprepRecipe.h"
 #include "Parameterization/DataprepParameterization.h"
 
 #include "AssetRegistryModule.h"
@@ -32,24 +31,10 @@
 
 // UDataprepAsset =================================================================
 
-UDataprepAsset::UDataprepAsset()
-{
-#ifndef DP_NOBLUEPRINT
-	DataprepRecipeBP = nullptr;
-	StartNode = nullptr;
-#endif
-	CachedActionCount = 0;
-}
-
 void UDataprepAsset::PostLoad()
 {
 	UDataprepAssetInterface::PostLoad();
 
-#ifndef NO_BLUEPRINT
-	check(DataprepRecipeBP);
-	DataprepRecipeBP->OnChanged().AddUObject( this, &UDataprepAsset::OnDataprepBlueprintChanged );
-#endif
-
 	// Move content of deprecated properties to the corresponding new ones.
 	if(HasAnyFlags(RF_WasLoaded))
 	{
@@ -77,36 +62,13 @@ void UDataprepAsset::PostLoad()
 			bMarkDirty = true;
 		}
 
-#ifndef NO_BLUEPRINT
-		// Most likely a Dataprep asset from 4.23
-		if(StartNode == nullptr)
+		// Most likely a Dataprep asset from 4.24
+		if(DataprepRecipeBP_DEPRECATED != nullptr)
 		{
-			UEdGraph* PipelineGraph = FBlueprintEditorUtils::FindEventGraph(DataprepRecipeBP);
-			check( PipelineGraph );
-
-			for( UEdGraphNode* GraphNode : PipelineGraph->Nodes )
-			{
-				StartNode = Cast<UK2Node_DataprepProducer>(GraphNode);
-				if( StartNode )
-				{
-					break;
-				}
-			}
-
-			// This Dataprep asset was never opened in the editor
-			if( StartNode == nullptr && ActionAssets.Num() == 0)
-			{
-				IBlueprintNodeBinder::FBindingSet Bindings;
-				StartNode = UBlueprintNodeSpawner::Create<UK2Node_DataprepProducer>()->Invoke( PipelineGraph, Bindings, FVector2D(-100,0) );
-				check(Cast<UK2Node_DataprepProducer>(StartNode));
-
-				DataprepRecipeBP->MarkPackageDirty();
-			}
-
-			UpdateActions(false);
+			DataprepRecipeBP_DEPRECATED = nullptr;
+			StartNode_DEPRECATED = nullptr;
 			bMarkDirty = true;
 		}
-#endif
 
 		if ( !Parameterization )
 		{
@@ -124,8 +86,6 @@ void UDataprepAsset::PostLoad()
 
 			GetOutermost()->SetDirtyFlag(true);
 		}
-
-		CachedActionCount = ActionAssets.Num();
 	}
 }
 
@@ -138,24 +98,33 @@ bool UDataprepAsset::Rename(const TCHAR* NewName/* =nullptr */, UObject* NewOute
 		{
 			bWasRename &= Parameterization->OnAssetRename( Flags );
 		}
-
-		if ( DataprepRecipeBP && bWasRename )
-		{
-			// There shouldn't be a blueprint depending on us. Should be ok to just rename the generated class
-			bWasRename &= DataprepRecipeBP->RenameGeneratedClasses( NewName, NewOuter, Flags );
-		}
 	}
 
 	return bWasRename;
 }
 
+void UDataprepAsset::PreEditUndo()
+{
+	UDataprepAssetInterface::PreEditUndo();
+
+	// Cache number of actions to detect addition/removal
+	ActionCountBeforeUndoRedo = ActionAssets.Num();
+}
+
 void UDataprepAsset::PostEditUndo()
 {
 	UDataprepAssetInterface::PostEditUndo();
 
-	OnActionChanged.Broadcast(nullptr, (ActionAssets.Num() == CachedActionCount) ? FDataprepAssetChangeType::ActionMoved : FDataprepAssetChangeType::ActionRemoved);
+	if(ActionCountBeforeUndoRedo != INDEX_NONE)
+	{
+		OnActionChanged.Broadcast(nullptr, (ActionAssets.Num() == ActionCountBeforeUndoRedo) ? FDataprepAssetChangeType::ActionMoved : FDataprepAssetChangeType::ActionRemoved);
+	}
+	else
+	{
+		ensure(false);
+	}
 
-	CachedActionCount = ActionAssets.Num();
+	ActionCountBeforeUndoRedo = INDEX_NONE;
 }
 
 void UDataprepAsset::PostDuplicate(EDuplicateMode::Type DuplicateMode)
@@ -187,211 +156,6 @@ const UDataprepActionAsset* UDataprepAsset::GetAction(int32 Index) const
 	return nullptr;
 }
 
-#ifndef NO_BLUEPRINT
-void UDataprepAsset::RemoveActionUsingBP(int32 Index)
-{
-	if ( ActionAssets.IsValidIndex( Index ) )
-	{
-		if ( UDataprepActionAsset* DataprepActionAsset = ActionAssets[Index] )
-		{
-			// Note this code will need to be updated with the new graph (also performance wise it's not really good (to many events) )
-			if ( UK2Node_DataprepActionCore* DataprepActionNode = Cast<UK2Node_DataprepActionCore>( DataprepActionAsset->GetOuter() ) )
-			{
-				UEdGraphPin* OuputPin = DataprepActionNode->FindPin( UEdGraphSchema_K2::PN_Then, EGPD_Output );
-				UEdGraphPin* InputPin = DataprepActionNode->FindPin( UEdGraphSchema_K2::PN_Execute, EGPD_Input );
-
-				// Reconnects the input of the node to it's output
-				if ( OuputPin && InputPin && OuputPin->LinkedTo.Num() > 0 && InputPin->LinkedTo.Num() )
-				{
-					if ( const UEdGraphSchema* GraphSchema = DataprepActionNode->GetSchema() )
-					{
-						TArray<UEdGraphPin*> Froms = InputPin->LinkedTo;
-						UEdGraphPin* To = OuputPin->LinkedTo[0];
-
-						// Notification will be send latter for the froms and to (modification are still recorded if there is a transaction)
-						constexpr bool bSendNotification = false;
-						GraphSchema->BreakPinLinks( *InputPin, bSendNotification );
-						GraphSchema->BreakPinLinks( *OuputPin, bSendNotification );
-
-						for ( UEdGraphPin* From : Froms )
-						{
-							GraphSchema->TryCreateConnection( From, To );
-						}
-					}
-				}
-
-				DataprepActionNode->DestroyNode();
-				UpdateActions();
-			}
-		}
-	}
-	else
-	{
-		UE_LOG( LogDataprepCore
-			, Error
-			,TEXT("The action to remove is out of bound. (Passed index: %d, Number of actions: %d, Dataprepsset: %s)")
-			, Index
-			, ActionAssets.Num()
-			, *GetPathName()
-			);
-	}
-
-}
-
-UDataprepActionAsset* UDataprepAsset::AddActionUsingBP(UEdGraphNode* NewActionNode)
-{
-	if ( ActionAssets.Num() > 0 )
-	{
-		UDataprepActionAsset* LastDataprepAction = ActionAssets.Last();
-		if ( UEdGraphNode* LastActionNode = Cast<UEdGraphNode>( LastDataprepAction->GetOuter() ) )
-		{
-			NewActionNode->AutowireNewNode( LastActionNode->FindPin( UEdGraphSchema_K2::PN_Then, EGPD_Output ) );
-		}
-	}
-	else
-	{
-		NewActionNode->AutowireNewNode( StartNode->FindPin( UEdGraphSchema_K2::PN_Then, EGPD_Output ) );
-	}
-	UpdateActions();
-
-	//Todo return the action
-	return nullptr;
-}
-
-void UDataprepAsset::SwapActionsUsingBP(int32 FirstActionIndex, int32 SecondActionIndex)
-{
-	if ( ActionAssets.Num() > 0 )
-	{
-		if ( !ActionAssets.IsValidIndex( FirstActionIndex ) || !ActionAssets.IsValidIndex( SecondActionIndex ) )
-		{
-			UE_LOG( LogDataprepCore
-				, Error
-				, TEXT("Can swap the dataprep actions a index is out of range. (First Index : %d, Second Index: %d, Number of Actions: %d, DataprepAction: %s)")
-				, FirstActionIndex
-				, SecondActionIndex
-				, ActionAssets.Num()
-				, *GetPathName()
-				);
-		}
-
-		// Note this code will need to be updated with the new graph (also performance wise it's not really good (to many events) )
-		auto GetOutput = [](UEdGraphPin& OuputPin) -> UEdGraphPin*
-		{
-			if (OuputPin.LinkedTo.Num() > 0)
-			{
-				return OuputPin.LinkedTo[0];
-			}
-			return nullptr;
-		};
-
-		// Grab the in/out of the first action
-		UDataprepActionAsset* FirstDataprepActionAsset = ActionAssets[FirstActionIndex];
-		check(FirstDataprepActionAsset);
-		UK2Node_DataprepActionCore* FirstDataprepActionNode = Cast<UK2Node_DataprepActionCore>( FirstDataprepActionAsset->GetOuter() );
-		check(FirstDataprepActionNode);
-
-		UEdGraphPin* FirstOuputPin = FirstDataprepActionNode->FindPin( UEdGraphSchema_K2::PN_Then, EGPD_Output );
-		check(FirstOuputPin);
-		UEdGraphPin* FirstOuput = GetOutput( *FirstOuputPin );
-
-		UEdGraphPin* FirstInputPin = FirstDataprepActionNode->FindPin( UEdGraphSchema_K2::PN_Execute, EGPD_Input );
-		check(FirstInputPin);
-		TArray<UEdGraphPin*> FirstInputs = FirstInputPin->LinkedTo;
-
-		// Grab the in/out of the second action
-		UDataprepActionAsset* SecondDataprepActionAsset = ActionAssets[SecondActionIndex];
-		check(SecondDataprepActionAsset);
-		UK2Node_DataprepActionCore* SecondDataprepActionNode = Cast<UK2Node_DataprepActionCore>( SecondDataprepActionAsset->GetOuter() );
-		check(SecondDataprepActionNode);
-
-		UEdGraphPin* SecondOuputPin = SecondDataprepActionNode->FindPin( UEdGraphSchema_K2::PN_Then, EGPD_Output );
-		check(SecondOuputPin);
-		UEdGraphPin* SecondOuput = GetOutput( *SecondOuputPin );
-
-		UEdGraphPin* SecondInputPin = SecondDataprepActionNode->FindPin( UEdGraphSchema_K2::PN_Execute, EGPD_Input );
-		check(SecondInputPin);
-		TArray<UEdGraphPin*> SecondInputs = SecondInputPin->LinkedTo;
-
-		// Reconnect the nodes
-		// Notification will be send latter for the froms and to (modification are still recorded if there is a transaction)
-		constexpr bool bSendNotification = false;
-		FirstOuputPin->BreakAllPinLinks( bSendNotification );
-		FirstInputPin->BreakAllPinLinks( bSendNotification );
-		SecondOuputPin->BreakAllPinLinks( bSendNotification );
-		SecondInputPin->BreakAllPinLinks( bSendNotification );
-
-		const UEdGraphSchema* GraphSchema = FirstDataprepActionNode->GetSchema();
-		check(GraphSchema);
-		if ( FMath::Abs( FirstActionIndex - SecondActionIndex ) == 1 )
-		{
-			for (UEdGraphPin* FirstInput : FirstInputs)
-			{
-				GraphSchema->TryCreateConnection( FirstInput, SecondInputPin );
-			}
-			GraphSchema->TryCreateConnection( FirstOuputPin, SecondOuput );
-			GraphSchema->TryCreateConnection( SecondOuputPin, FirstInputPin );
-		}
-		else
-		{
-			GraphSchema->TryCreateConnection(FirstOuputPin, SecondOuput);
-			for (UEdGraphPin* SecondInput : SecondInputs)
-			{
-				GraphSchema->TryCreateConnection(SecondInput, FirstInputPin);
-			}
-			GraphSchema->TryCreateConnection(SecondOuputPin, FirstOuput);
-			for (UEdGraphPin* FirstInput : FirstInputs)
-			{
-				GraphSchema->TryCreateConnection(FirstInput, SecondInputPin);
-			}
-		}
-
-		UpdateActions();
-	}
-	else
-	{
-		UE_LOG( LogDataprepCore
-			, Error
-			, TEXT("Can't swap the actions of a DataprepAsset without actions. (DataprepAsset: %s)")
-			, *GetPathName()
-			);
-	}
-}
-
-bool UDataprepAsset::CreateBlueprint()
-{
-	// Begin: Temp code for the nodes development
-	const FString DesiredName = GetName() + TEXT("_Recipe");
-	FName BlueprintName = MakeUniqueObjectName( GetOutermost(), UBlueprint::StaticClass(), *DesiredName );
-
-	DataprepRecipeBP = FKismetEditorUtilities::CreateBlueprint( UDataprepRecipe::StaticClass(), this, BlueprintName, BPTYPE_Normal, UBlueprint::StaticClass(), UBlueprintGeneratedClass::StaticClass() );
-	check( DataprepRecipeBP );
-
-	// This blueprint is not the asset of the package
-	DataprepRecipeBP->ClearFlags( RF_Standalone );
-
-	FAssetRegistryModule::AssetCreated( DataprepRecipeBP );
-
-	// Create the start node of the Blueprint
-	UEdGraph* PipelineGraph = FBlueprintEditorUtils::FindEventGraph(DataprepRecipeBP);
-	check( PipelineGraph );
-
-	UEdGraph* EventGraph = FBlueprintEditorUtils::FindEventGraph(DataprepRecipeBP);
-	IBlueprintNodeBinder::FBindingSet Bindings;
-
-	StartNode = UBlueprintNodeSpawner::Create<UK2Node_DataprepProducer>()->Invoke( EventGraph, Bindings, FVector2D(-100,0) );
-	check(Cast<UK2Node_DataprepProducer>(StartNode));
-
-	DataprepRecipeBP->MarkPackageDirty();
-
-	DataprepRecipeBP->OnChanged().AddUObject( this, &UDataprepAsset::OnDataprepBlueprintChanged );
-	// End: Temp code for the nodes development
-
-	MarkPackageDirty();
-
-	return true;
-}
-#endif
-
 bool UDataprepAsset::CreateParameterization()
 {
 	if( !Parameterization )
@@ -486,80 +250,6 @@ void UDataprepAsset::FRestrictedToActionAsset::NotifyAssetOfTheRemovalOfSteps(UD
 	DataprepAsset.OnStepObjectsAboutToBeRemoved.Broadcast( StepObjects );
 }
 
-#ifndef NO_BLUEPRINT
-void UDataprepAsset::OnDataprepBlueprintChanged( UBlueprint* InBlueprint )
-{
-#if 0
-	if(InBlueprint == DataprepRecipeBP)
-	{
-		UpdateActions();
-		OnChanged.Broadcast( FDataprepAssetChangeType::RecipeModified );
-	}
-#endif
-}
-
-void UDataprepAsset::UpdateActions(bool bNotify)
-{
-	ActionAssets.Empty(ActionAssets.Num());
-
-	UEdGraphPin* NodeOutPin= StartNode->FindPin(UEdGraphSchema_K2::PN_Then, EGPD_Output);
-	if( NodeOutPin && NodeOutPin->LinkedTo.Num() > 0 )
-	{
-		TSet<UEdGraphNode*> ActionNodesVisited;
-
-		for( UEdGraphPin* NextNodeInPin = NodeOutPin->LinkedTo[0]; NextNodeInPin != nullptr ; )
-		{
-			UEdGraphNode* NextNode = NextNodeInPin->GetOwningNode();
-
-			uint32 NodeHash = GetTypeHash( NextNode );
-			// Break the loop if the node had already been visited
-			if( ActionNodesVisited.FindByHash( NodeHash, NextNode ) )
-			{
-				break;
-			}
-			else
-			{
-				ActionNodesVisited.AddByHash( NodeHash, NextNode );
-			}
-
-			if(UK2Node_DataprepActionCore* ActionNode = Cast<UK2Node_DataprepActionCore>(NextNode))
-			{
-				if( UDataprepActionAsset* DataprepAction = ActionNode->GetDataprepAction() )
-				{
-					DataprepAction->Rename(nullptr, this, REN_DoNotDirty | REN_DontCreateRedirectors | REN_NonTransactional);
-					ActionAssets.Add( DataprepAction );
-				}
-			}
-
-			// Look for the next node
-			NodeOutPin = NextNode->FindPin( UEdGraphSchema_K2::PN_Then, EGPD_Output );
-
-			if ( !NodeOutPin )
-			{
-				// If we couldn't find a then pin try to get the first output pin as a fallback
-				for ( UEdGraphPin* Pin : NextNode->Pins )
-				{
-					if ( Pin && Pin->PinType.PinCategory == UEdGraphSchema_K2::PC_Exec && Pin->Direction == EGPD_Output )
-					{
-						NodeOutPin = Pin;
-						break;
-					}
-				}
-			}
-
-			NextNodeInPin = NodeOutPin ? ( NodeOutPin->LinkedTo.Num() > 0 ? NodeOutPin->LinkedTo[0] : nullptr ) : nullptr;
-		}
-	}
-
-	if(bNotify)
-	{
-		OnActionChanged.Broadcast(ActionAssets.Num() > 0 ? ActionAssets[0] : nullptr, FDataprepAssetChangeType::ActionAdded);
-	}
-
-	CachedActionCount = ActionAssets.Num();
-}
-#endif
-
 int32 UDataprepAsset::AddAction(const UDataprepActionAsset* InAction)
 {
 	UDataprepActionAsset* Action = InAction ? DuplicateObject<UDataprepActionAsset>( InAction, this) : NewObject<UDataprepActionAsset>( this, UDataprepActionAsset::StaticClass(), NAME_None, RF_Transactional );
@@ -574,8 +264,6 @@ int32 UDataprepAsset::AddAction(const UDataprepActionAsset* InAction)
 		ActionAssets.Add( Action );
 		OnActionChanged.Broadcast(Action, FDataprepAssetChangeType::ActionAdded);
 
-		CachedActionCount = ActionAssets.Num();
-
 		return ActionAssets.Num() - 1;
 	}
 
@@ -606,9 +294,7 @@ int32 UDataprepAsset::AddActions(const TArray<const UDataprepActionAsset*>& InAc
 			}
 		}
 
-		CachedActionCount = ActionAssets.Num();
-
-		if(PreviousActionCount != CachedActionCount)
+		if(PreviousActionCount != ActionAssets.Num())
 		{
 			OnActionChanged.Broadcast(ActionAssets.Last(), FDataprepAssetChangeType::ActionAdded);
 
@@ -636,8 +322,6 @@ int32 UDataprepAsset::AddAction(const TArray<const UDataprepActionStep*>& InActi
 
 		Action->AddSteps(InActionSteps);
 
-		CachedActionCount = ActionAssets.Num();
-
 		OnActionChanged.Broadcast(Action, FDataprepAssetChangeType::ActionAdded);
 
 		return ActionAssets.Num() - 1;
@@ -670,8 +354,6 @@ bool UDataprepAsset::InsertAction(const UDataprepActionAsset* InAction, int32 In
 
 		OnActionChanged.Broadcast( Action, FDataprepAssetChangeType::ActionAdded );
 
-		CachedActionCount = ActionAssets.Num();
-
 		return true;
 	}
 
@@ -697,12 +379,13 @@ bool UDataprepAsset::InsertActions(const TArray<const UDataprepActionAsset*>& In
 		int32 PreviousActionCount = ActionAssets.Num();
 
 		int32 InsertIndex = Index;
+		UDataprepActionAsset* Action = nullptr;
 
 		for(const UDataprepActionAsset* InAction : InActions)
 		{
 			if(InAction)
 			{
-				UDataprepActionAsset* Action = DuplicateObject<UDataprepActionAsset>( InAction, this);
+				Action = DuplicateObject<UDataprepActionAsset>( InAction, this);
 				Action->SetFlags(EObjectFlags::RF_Transactional);
 				Action->SetLabel( InAction->GetLabel() );
 
@@ -712,12 +395,9 @@ bool UDataprepAsset::InsertActions(const TArray<const UDataprepActionAsset*>& In
 			}
 		}
 
-		CachedActionCount = ActionAssets.Num();
-
-		if(PreviousActionCount != CachedActionCount)
+		if(Action != nullptr)
 		{
-			// We should rework this notification
-			OnActionChanged.Broadcast(ActionAssets.Last(), FDataprepAssetChangeType::ActionAdded);
+			OnActionChanged.Broadcast(Action, FDataprepAssetChangeType::ActionAdded);
 
 			return true;
 		}
@@ -739,8 +419,6 @@ bool UDataprepAsset::InsertAction(int32 Index)
 
 		ActionAssets.Insert( Action, Index );
 
-		CachedActionCount = ActionAssets.Num();
-
 		OnActionChanged.Broadcast( Action, FDataprepAssetChangeType::ActionAdded );
 		return true;
 	}
@@ -767,8 +445,6 @@ bool UDataprepAsset::InsertAction(const TArray<const UDataprepActionStep*>& InAc
 
 			Action->AddSteps( InActionSteps );
 
-			CachedActionCount = ActionAssets.Num();
-
 			OnActionChanged.Broadcast( Action, FDataprepAssetChangeType::ActionAdded );
 			return true;
 		}
@@ -793,12 +469,12 @@ bool UDataprepAsset::MoveAction(int32 SourceIndex, int32 DestinationIndex)
 	{
 		if ( !ActionAssets.IsValidIndex( SourceIndex ) )
 		{
-			UE_LOG( LogDataprepCore, Error, TEXT("UDataprepAsset::MoveAction: The Step Index is out of range") );
+			UE_LOG( LogDataprepCore, Error, TEXT("UDataprepAsset::MoveAction: The source index is out of range") );
 		}
 
 		if ( !ActionAssets.IsValidIndex( DestinationIndex ) )
 		{
-			UE_LOG( LogDataprepCore, Error, TEXT("UDataprepAsset::MoveAction: The Destination Index is out of range") );
+			UE_LOG( LogDataprepCore, Error, TEXT("UDataprepAsset::MoveAction: The destination index is out of range") );
 		}
 		return false;
 	}
@@ -815,6 +491,39 @@ bool UDataprepAsset::MoveAction(int32 SourceIndex, int32 DestinationIndex)
 	return false;
 }
 
+bool UDataprepAsset::SwapActions(int32 FirstActionIndex, int32 SecondActionIndex)
+{
+	if ( FirstActionIndex == SecondActionIndex )
+	{
+		UE_LOG( LogDataprepCore, Error, TEXT("UDataprepAsset::SwapAction: Nothing done. First and second indices are identical") );
+		return true;
+	}
+
+	if ( !ActionAssets.IsValidIndex( FirstActionIndex ) || !ActionAssets.IsValidIndex( SecondActionIndex ) )
+	{
+		if ( !ActionAssets.IsValidIndex( FirstActionIndex ) )
+		{
+			UE_LOG( LogDataprepCore, Error, TEXT("UDataprepAsset::SwapAction: The first index is out of range") );
+		}
+
+		if ( !ActionAssets.IsValidIndex( SecondActionIndex ) )
+		{
+			UE_LOG( LogDataprepCore, Error, TEXT("UDataprepAsset::SwapAction: The second index is out of range") );
+		}
+
+		return false;
+	}
+
+	Modify();
+
+	// Directly use TArray::SwapMemory since we have validated the indices
+	ActionAssets.SwapMemory(FirstActionIndex, SecondActionIndex);
+
+	OnActionChanged.Broadcast(ActionAssets[FirstActionIndex], FDataprepAssetChangeType::ActionMoved);
+
+	return true;
+}
+
 bool UDataprepAsset::RemoveAction(int32 Index)
 {
 	if ( ActionAssets.IsValidIndex( Index ) )
@@ -829,8 +538,6 @@ bool UDataprepAsset::RemoveAction(int32 Index)
 
 		ActionAssets.RemoveAt( Index );
 
-		CachedActionCount = ActionAssets.Num();
-
 		OnActionChanged.Broadcast(ActionAsset, FDataprepAssetChangeType::ActionRemoved);
 
 		return true;
@@ -880,8 +587,6 @@ bool UDataprepAsset::RemoveActions(const TArray<int32>& Indices)
 			}
 		}
 
-		CachedActionCount = ActionAssets.Num();
-
 		// Notify on last action removed
 		OnActionChanged.Broadcast(ActionAsset, FDataprepAssetChangeType::ActionRemoved);
 
diff --git a/Engine/Plugins/Enterprise/DataprepEditor/Source/DataprepCore/Private/DataprepAssetInterface.cpp b/Engine/Plugins/Enterprise/DataprepEditor/Source/DataprepCore/Private/DataprepAssetInterface.cpp
--- a/Engine/Plugins/Enterprise/DataprepEditor/Source/DataprepCore/Private/DataprepAssetInterface.cpp
+++ b/Engine/Plugins/Enterprise/DataprepEditor/Source/DataprepCore/Private/DataprepAssetInterface.cpp
@@ -8,7 +8,6 @@
 #include "DataprepCoreLogCategory.h"
 #include "DataprepCorePrivateUtils.h"
 #include "DataprepCoreUtils.h"
-#include "DataprepRecipe.h"
 
 #include "AssetRegistryModule.h"
 #ifdef WITH_EDITOR
diff --git a/Engine/Plugins/Enterprise/DataprepEditor/Source/DataprepCore/Private/DataprepFactories.cpp b/Engine/Plugins/Enterprise/DataprepEditor/Source/DataprepCore/Private/DataprepFactories.cpp
--- a/Engine/Plugins/Enterprise/DataprepEditor/Source/DataprepCore/Private/DataprepFactories.cpp
+++ b/Engine/Plugins/Enterprise/DataprepEditor/Source/DataprepCore/Private/DataprepFactories.cpp
@@ -84,11 +84,6 @@ UObject * UDataprepAssetFactory::FactoryCreateNew(UClass* InClass, UObject* InPa
 	}
 	check( DataprepAsset->GetConsumer() );
 
-	// Temp code for the nodes development
-	// Initialize Dataprep asset's blueprint
-	DataprepAsset->CreateBlueprint();
-	// end of temp code for nodes development
-
 	DataprepAsset->CreateParameterization();
 
 	FAssetRegistryModule::AssetCreated( DataprepAsset );
diff --git a/Engine/Plugins/Enterprise/DataprepEditor/Source/DataprepCore/Private/DataprepRecipe.cpp b/Engine/Plugins/Enterprise/DataprepEditor/Source/DataprepCore/Private/DataprepRecipe.cpp
--- a/Engine/Plugins/Enterprise/DataprepEditor/Source/DataprepCore/Private/DataprepRecipe.cpp
+++ b/Engine/Plugins/Enterprise/DataprepEditor/Source/DataprepCore/Private/DataprepRecipe.cpp
@@ -1,106 +1,3 @@
 // Copyright Epic Games, Inc. All Rights Reserved.
 
 #include "DataprepRecipe.h"
-
-#include "ActorEditorUtils.h"
-#include "Engine/World.h"
-#include "Engine/StaticMesh.h"
-#include "EngineUtils.h"
-#include "GameFramework/Actor.h"
-#include "Materials/Material.h"
-#include "UObject/SoftObjectPath.h"
-
-#define LOCTEXT_NAMESPACE "DataprepRecipe"
-
-namespace DataprepRecipeUtils
-{
-	bool IsActorValid(const AActor* Actor)
-	{
-		if (Actor == nullptr)
-		{
-			return false;
-		}
-
-		// Don't consider transient actors in non-play worlds
-		// Don't consider the builder brush
-		// Don't consider the WorldSettings actor, even though it is technically editable
-		bool bIsValid = Actor->IsEditable() && !Actor->IsTemplate() && !Actor->HasAnyFlags(RF_Transient) && !FActorEditorUtils::IsABuilderBrush(Actor) && !Actor->IsA(AWorldSettings::StaticClass());
-
-		if (Actor->GetWorld() == GWorld)
-		{
-			// Only consider actors that are allowed to be selected and drawn in editor
-			bIsValid &= Actor->IsListedInSceneOutliner();
-		}
-
-		return bIsValid;
-	}
-}
-
-UDataprepRecipe::UDataprepRecipe()
-{
-	TargetWorld = nullptr;
-}
-
-TArray<TWeakObjectPtr<UObject>> UDataprepRecipe::GetValidAssets(bool bFlushAssets /*= true*/)
-{
-	TArray<TWeakObjectPtr<UObject>> ValidAssets;
-
-#if WITH_EDITOR
-	for (const TWeakObjectPtr<UObject>& AssetPtr : Assets)
-	{
-		if (AssetPtr.Get() != nullptr && !AssetPtr->IsPendingKill())
-		{
-			ValidAssets.Add(AssetPtr);
-		}
-	}
-
-	// Empty data prep recipe's list of assets
-	if (bFlushAssets == true)
-	{
-		Assets.Empty();
-	}
-#endif
-
-	return ValidAssets;
-}
-
-TArray<AActor*> UDataprepRecipe::GetActors() const
-{
-	TArray<AActor*> Result;
-
-	if (TargetWorld == nullptr)
-	{
-		return Result;
-	}
-
-#if WITH_EDITORONLY_DATA
-	const EActorIteratorFlags Flags = EActorIteratorFlags::SkipPendingKill;
-	for (TActorIterator<AActor> It(TargetWorld, AActor::StaticClass(), Flags); It; ++It)
-	{
-		AActor* Actor = *It;
-		if (DataprepRecipeUtils::IsActorValid(Actor))
-		{
-			Result.Add(Actor);
-		}
-	}
-#endif
-
-	return Result;
-}
-
-TArray<UObject*> UDataprepRecipe::GetAssets() const
-{
-	TArray<UObject*> Result;
-
-	for (const TWeakObjectPtr<UObject>& AssetPtr : Assets)
-	{
-		if (AssetPtr.Get() != nullptr && !AssetPtr->IsPendingKill())
-		{
-			Result.Add(AssetPtr.Get());
-		}
-	}
-
-	return Result;
-}
-
-#undef LOCTEXT_NAMESPACE
diff --git a/Engine/Plugins/Enterprise/DataprepEditor/Source/DataprepCore/Public/DataprepAsset.h b/Engine/Plugins/Enterprise/DataprepEditor/Source/DataprepCore/Public/DataprepAsset.h
--- a/Engine/Plugins/Enterprise/DataprepEditor/Source/DataprepCore/Public/DataprepAsset.h
+++ b/Engine/Plugins/Enterprise/DataprepEditor/Source/DataprepCore/Public/DataprepAsset.h
@@ -10,9 +10,6 @@
 
 #include "DataprepAsset.generated.h"
 
-#ifndef NO_BLUEPRINT
-class UBlueprint;
-#endif
 class UDataprepActionAsset;
 class UDataprepActionStep;
 class UDataprepParameterizableObject;
@@ -35,13 +32,18 @@ class DATAPREPCORE_API UDataprepAsset : public UDataprepAssetInterface
 	GENERATED_BODY()
 
 public:
-	UDataprepAsset();
+	UDataprepAsset()
+		: ActionCountBeforeUndoRedo(INDEX_NONE)
+	{
+	}
+
 
 	virtual ~UDataprepAsset() = default;
 
 	// UObject interface
 	virtual void PostLoad() override;
 	virtual bool Rename(const TCHAR* NewName/* =nullptr */, UObject* NewOuter/* =nullptr */, ERenameFlags Flags/* =REN_None */) override;
+	virtual void PreEditUndo() override;
 	virtual void PostEditUndo() override;
 	virtual void PostDuplicate(EDuplicateMode::Type DuplicateMode) override;
 	// End of UObject interface
@@ -136,6 +138,13 @@ public:
 	*/
 	bool MoveAction(int32 SourceIndex, int32 DestinationIndex);
 
+	/**
+	 * Swap the actions of a Dataprep asset
+	 * @param FirstActionIndex The index of the first action
+	 * @param SecondActionIndex The index of the second action
+	 */
+	bool SwapActions(int32 FirstActionIndex, int32 SecondActionIndex);
+
 	/**
 	 * Remove an action from the Dataprep asset
 	 * @param Index The index of the action to remove
@@ -172,36 +181,11 @@ public:
 
 	bool CreateParameterization();
 
-#ifndef NO_BLUEPRINT
-	// Temp code for the nodes development
-	bool CreateBlueprint();
-
-	/** @return pointer on the recipe */
-	const UBlueprint* GetRecipeBP() const
-	{ 
-		return DataprepRecipeBP;
-	}
-
-	UBlueprint* GetRecipeBP()
-	{ 
-		return DataprepRecipeBP;
-	}
-
-	//Todo Change the signature of this function when the new graph is ready (Hack to avoid a refactoring)
-	UDataprepActionAsset* AddActionUsingBP(class UEdGraphNode* NewActionNode);
-
-	void SwapActionsUsingBP(int32 FirstActionIndex, int32 SecondActionIndex);
-
-	void RemoveActionUsingBP(int32 Index);
-
-	// end of temp code for nodes development
-#endif
-
 public:
 	// Functions specific to the parametrization of the Dataprep asset
 
 	/**
-	 * Event to notify the ui that a dataprep parametrization was modified
+	 * Event to notify the ui that a Dataprep parametrization was modified
 	 * This necessary as the ui for the parameterization is only updated by manual event (the ui don't pull new values each frame)
 	 * Note on the objects param: The parameterized objects that should refresh their ui. If nullptr all widgets that can display some info on the parameterization should be refreshed
 	 */
@@ -226,32 +210,24 @@ public:
 
 protected:
 #if WITH_EDITORONLY_DATA
-	// Temp code for the nodes development
-	/** Temporary: Pointer to data preparation pipeline blueprint used to process input data */
-	UPROPERTY()
-	UBlueprint* DataprepRecipeBP;
+	/** DEPRECATED: Pointer to data preparation pipeline blueprint previously used to process input data */
+	UPROPERTY(meta = (DeprecatedProperty, DeprecationMessage = "Using directly ActionAssets property instead of Blueprint to manage actions."))
+	UBlueprint* DataprepRecipeBP_DEPRECATED;
 	// end of temp code for nodes development
 
 	/** DEPRECATED: List of producers referenced by the asset */
-	UPROPERTY()
+	UPROPERTY(meta = (DeprecatedProperty, DeprecationMessage = "Property moved to UDataprepAssetInterface as Inputs."))
 	TArray< FDataprepAssetProducer > Producers_DEPRECATED;
 
 	/** DEPRECATED: COnsumer referenced by the asset */
-	UPROPERTY()
+	UPROPERTY(meta = (DeprecatedProperty, DeprecationMessage = "Property moved to UDataprepAssetInterface as Output."))
 	UDataprepContentConsumer* Consumer_DEPRECATED;
 #endif
 
-	// Temp code for the nodes development
-	void OnDataprepBlueprintChanged( class UBlueprint* InBlueprint );
-	// end of temp code for nodes development
-
-	// Temp code for the nodes development
 private:
-	void UpdateActions(bool bNotify = true);
-
-private:
-	UPROPERTY()
-	UEdGraphNode* StartNode;
+	/** DEPRECATED: Pointer to the entry node of the pipeline blueprint previously used to process input data */
+	UPROPERTY(meta = (DeprecatedProperty, DeprecationMessage = "Using directly ActionAssets property instead of Blueprint to manage actions."))
+	UEdGraphNode* StartNode_DEPRECATED;
 
 	UPROPERTY()
 	UDataprepParameterization* Parameterization;
@@ -263,5 +239,6 @@ private:
 
 	FOnStepObjectsAboutToBeRemoved OnStepObjectsAboutToBeRemoved;
 
-	int32 CachedActionCount;
+	/** Used to identify addition/removal of actions on undo/redo */
+	int32 ActionCountBeforeUndoRedo;
 };
\ No newline at end of file
diff --git a/Engine/Plugins/Enterprise/DataprepEditor/Source/DataprepCore/Public/DataprepRecipe.h b/Engine/Plugins/Enterprise/DataprepEditor/Source/DataprepCore/Public/DataprepRecipe.h
--- a/Engine/Plugins/Enterprise/DataprepEditor/Source/DataprepCore/Public/DataprepRecipe.h
+++ b/Engine/Plugins/Enterprise/DataprepEditor/Source/DataprepCore/Public/DataprepRecipe.h
@@ -10,86 +10,44 @@
 class AActor;
 class UWorld;
 
-UCLASS(Experimental, Blueprintable, BlueprintType, meta = (DisplayName = "Data Preparation Recipe"))
-class DATAPREPCORE_API UDataprepRecipe : public UObject
+UCLASS(Deprecated, Experimental, Blueprintable, BlueprintType, meta = (DisplayName = "DEPRECATED Data Preparation Recipe", DeprecationMessage = "No use of Blueprint with Dataprep."))
+class DATAPREPCORE_API UDEPRECATED_DataprepRecipe : public UObject
 {
 	GENERATED_BODY()
 
 public:
-	UDataprepRecipe();
+	UDEPRECATED_DataprepRecipe() {}
 
 	/** Begin UObject override */
 	virtual bool IsEditorOnly() const override { return true; }
 	/** End UObject override */
 
+public:
 	/**
-	 * Reset array of assets attached to recipe
-	 */
-	void ResetAssets()
-	{
-		Assets.Reset();
-	}
-
-	/**
-	 * Setter on array of assets held by this object
+	 * DEPRECATED
+	 * Returns all actors contained in its attached world
 	 */
-	void SetAssets(TArray<TWeakObjectPtr<UObject>> InAssets)
+	UFUNCTION(BlueprintCallable, Category = "Dataprep | Query", meta = (HideSelfPin = "true"), meta=(DeprecatedFunction, DeprecationMessage="No use of Blueprint with Dataprep."))
+	virtual TArray<AActor*> GetActors() const
 	{
-		Assets = MoveTemp(InAssets);
+		return TArray<AActor*>();
 	}
 
 	/**
-	 * Returns array of assets attached to recipe.
-	 * Usually called after the operations of the recipe have been executed
-	 */
-	TArray<TWeakObjectPtr<UObject>> GetValidAssets(bool bFlushAssets = true);
-
-	/**
-	 * Returns world attached to this recipe
+	 * DEPRECATED
+	 * Returns all valid assets contained in attached world
 	 */
-	UWorld* GetTargetWorld()
+	UFUNCTION(BlueprintCallable, Category = "Dataprep | Query", meta = (DeprecatedFunction, DeprecationMessage="No use of Blueprint with Dataprep.", HideSelfPin = "true"))
+	virtual TArray<UObject*> GetAssets() const
 	{
-		return TargetWorld;
+		return TArray<UObject*>();
 	}
 
 	/**
-	 * Attaches a world to this recipe
-	 */
-	void SetTargetWorld(UWorld* InTargetWorld)
-	{
-		TargetWorld = InTargetWorld;
-	}
-
-public:
-	/**
-	 * Returns all actors contained in its attached world
-	 */
-	UFUNCTION(BlueprintCallable, Category = "Dataprep | Query", meta = (HideSelfPin = "true"))
-	virtual TArray<AActor*> GetActors() const;
-
-	/**
-	 * Returns all valid assets contained in attached world
-	 */
-	UFUNCTION(BlueprintCallable, Category = "Dataprep | Query", meta = (HideSelfPin = "true"))
-	virtual TArray<UObject*> GetAssets() const;
-
-	/**
+	 * DEPRECATED
 	 * Function used to trigger the execution of the pipeline
 	 * An event node associated to this function must be in the pipeline graph to run it.
 	 */
-	UFUNCTION(BlueprintImplementableEvent)
+	UFUNCTION(BlueprintImplementableEvent, meta=(DeprecatedFunction, DeprecationMessage="No use of Blueprint with Dataprep."))
 	void TriggerPipelineTraversal();
-
-private:
-	/**
-	 * World attached to this recipe.
-	 * All queries on actors will be perform from this world.
-	 * Transient, Can be null.
-	 */
-	UWorld* TargetWorld;
-
-	/**
-	 * All queries on assets will be perform from this array.
-	 */
-	TArray<TWeakObjectPtr<UObject>> Assets;
 };
diff --git a/Engine/Plugins/Enterprise/DataprepEditor/Source/DataprepEditor/Private/AssetTypeActions_Dataprep.cpp b/Engine/Plugins/Enterprise/DataprepEditor/Source/DataprepEditor/Private/AssetTypeActions_Dataprep.cpp
deleted file mode 100644
diff --git a/Engine/Plugins/Enterprise/DataprepEditor/Source/DataprepEditor/Private/AssetTypeActions_Dataprep.h b/Engine/Plugins/Enterprise/DataprepEditor/Source/DataprepEditor/Private/AssetTypeActions_Dataprep.h
deleted file mode 100644
diff --git a/Engine/Plugins/Enterprise/DataprepEditor/Source/DataprepEditor/Private/AssetTypeActions_DataprepAsset.cpp b/Engine/Plugins/Enterprise/DataprepEditor/Source/DataprepEditor/Private/AssetTypeActions_DataprepAsset.cpp
--- a/Engine/Plugins/Enterprise/DataprepEditor/Source/DataprepEditor/Private/AssetTypeActions_DataprepAsset.cpp
+++ b/Engine/Plugins/Enterprise/DataprepEditor/Source/DataprepEditor/Private/AssetTypeActions_DataprepAsset.cpp
@@ -28,7 +28,7 @@ void FAssetTypeActions_DataprepAsset::OpenAssetEditor(const TArray<UObject*>& In
 		if (UDataprepAsset* DataprepAsset = Cast<UDataprepAsset>(Object))
 		{
 			TSharedRef<FDataprepEditor> NewDataprepEditor(new FDataprepEditor());
-			NewDataprepEditor->InitDataprepEditor( EToolkitMode::Standalone, EditWithinLevelEditor, DataprepAsset, DataprepAsset->GetRecipeBP() );
+			NewDataprepEditor->InitDataprepEditor( EToolkitMode::Standalone, EditWithinLevelEditor, DataprepAsset );
 		}
 	}
 }
diff --git a/Engine/Plugins/Enterprise/DataprepEditor/Source/DataprepEditor/Private/BlueprintNodes/K2Node_DataprepAction.cpp b/Engine/Plugins/Enterprise/DataprepEditor/Source/DataprepEditor/Private/BlueprintNodes/K2Node_DataprepAction.cpp
--- a/Engine/Plugins/Enterprise/DataprepEditor/Source/DataprepEditor/Private/BlueprintNodes/K2Node_DataprepAction.cpp
+++ b/Engine/Plugins/Enterprise/DataprepEditor/Source/DataprepEditor/Private/BlueprintNodes/K2Node_DataprepAction.cpp
@@ -6,7 +6,6 @@
 #include "DataprepAsset.h"
 #include "DataprepActionAsset.h"
 #include "DataprepEditorUtils.h"
-#include "Widgets/DataprepGraph/SGraphNodeK2DataprepAction.h"
 
 // Engine includes
 #include "BlueprintActionDatabaseRegistrar.h"
@@ -26,16 +25,16 @@
 
 const FText DataprepActionCategory( LOCTEXT( "DataprepActionK2NodeCategory", "Dataprep Action" ) );
 
-const FName UK2Node_DataprepAction::ThenPinName( TEXT("") );
-const FName UK2Node_DataprepAction::InObjectsPinName( TEXT("Objects") );
+const FName UDEPRECATED_K2Node_DataprepAction::ThenPinName( TEXT("") );
+const FName UDEPRECATED_K2Node_DataprepAction::InObjectsPinName( TEXT("Objects") );
 
-UK2Node_DataprepAction::UK2Node_DataprepAction()
+UDEPRECATED_K2Node_DataprepAction::UDEPRECATED_K2Node_DataprepAction()
 {
 	bCanRenameNode = true;
 	ActionTitle = LOCTEXT("DefaultNodeTitle", "New Action").ToString();
 }
 
-void UK2Node_DataprepAction::PostReconstructNode()
+void UDEPRECATED_K2Node_DataprepAction::PostReconstructNode()
 {
 	Super::PostReconstructNode();
 
@@ -53,7 +52,7 @@ void UK2Node_DataprepAction::PostReconstructNode()
 }
 
 
-void UK2Node_DataprepAction::AllocateDefaultPins()
+void UDEPRECATED_K2Node_DataprepAction::AllocateDefaultPins()
 {
 	// Inputs
 	CreatePin( EGPD_Input, UEdGraphSchema_K2::PC_Exec, UEdGraphSchema_K2::PN_Execute );
@@ -89,22 +88,22 @@ void UK2Node_DataprepAction::AllocateDefaultPins()
 	Super::AllocateDefaultPins();
 }
 
-FText UK2Node_DataprepAction::GetMenuCategory() const
+FText UDEPRECATED_K2Node_DataprepAction::GetMenuCategory() const
 {
 	return DataprepActionCategory;
 }
 
-FLinearColor UK2Node_DataprepAction::GetNodeTitleColor() const
+FLinearColor UDEPRECATED_K2Node_DataprepAction::GetNodeTitleColor() const
 {
 	return FLinearColor(0.0036765f, 0.3864294f, 0.2501584);
 }
 
-FText UK2Node_DataprepAction::GetNodeTitle(ENodeTitleType::Type TitleType) const
+FText UDEPRECATED_K2Node_DataprepAction::GetNodeTitle(ENodeTitleType::Type TitleType) const
 {
 	return FText::FromString( ActionTitle );
 }
 
-void UK2Node_DataprepAction::OnRenameNode(const FString& NewName)
+void UDEPRECATED_K2Node_DataprepAction::OnRenameNode(const FString& NewName)
 {
 	ActionTitle = NewName;
 	if ( DataprepActionAsset )
@@ -113,7 +112,7 @@ void UK2Node_DataprepAction::OnRenameNode(const FString& NewName)
 	}
 }
 
-void UK2Node_DataprepAction::DestroyNode()
+void UDEPRECATED_K2Node_DataprepAction::DestroyNode()
 {
 	if ( DataprepActionAsset )
 	{
@@ -128,23 +127,13 @@ void UK2Node_DataprepAction::DestroyNode()
 	Super::DestroyNode();
 }
 
-void UK2Node_DataprepAction::NodeConnectionListChanged()
-{
-	FDataprepEditorUtils::NotifySystemOfChangeInPipeline( GetDataprepAction() );
-}
-
-TSharedPtr<SGraphNode> UK2Node_DataprepAction::CreateVisualWidget()
-{
-	return SNew( SGraphNodeK2DataprepAction, this );
-}
-
-TSharedPtr<class INameValidatorInterface> UK2Node_DataprepAction::MakeNameValidator() const
+TSharedPtr<class INameValidatorInterface> UDEPRECATED_K2Node_DataprepAction::MakeNameValidator() const
 {
 	// The name doesn't matter
 	return MakeShareable(new FDummyNameValidator(EValidatorResult::Ok));
 }
 
-void UK2Node_DataprepAction::GetMenuActions(FBlueprintActionDatabaseRegistrar& ActionRegistrar)const
+void UDEPRECATED_K2Node_DataprepAction::GetMenuActions(FBlueprintActionDatabaseRegistrar& ActionRegistrar)const
 {
 	// actions get registered under specific object-keys; the idea is that 
 	// actions might have to be updated (or deleted) if their object-key is  
@@ -162,7 +151,7 @@ void UK2Node_DataprepAction::GetMenuActions(FBlueprintActionDatabaseRegistrar& A
  
 		auto CustomizeNodeLambda = [](UEdGraphNode* NewNode, bool bIsTemplateNode)
 		{
-			UK2Node_DataprepAction* DataprepActionNode  = CastChecked<UK2Node_DataprepAction>(NewNode);
+			UDEPRECATED_K2Node_DataprepAction* DataprepActionNode  = CastChecked<UDEPRECATED_K2Node_DataprepAction>(NewNode);
  
 			UBlueprint* Blueprint = DataprepActionNode->GetBlueprint();
 			if (Blueprint && Blueprint->GeneratedClass && !bIsTemplateNode )
@@ -177,7 +166,7 @@ void UK2Node_DataprepAction::GetMenuActions(FBlueprintActionDatabaseRegistrar& A
 	}
 }
 
-void UK2Node_DataprepAction::ExpandNode(FKismetCompilerContext& CompilerContext, UEdGraph* SourceGraph)
+void UDEPRECATED_K2Node_DataprepAction::ExpandNode(FKismetCompilerContext& CompilerContext, UEdGraph* SourceGraph)
 {
 	UK2Node_CallFunction* CallOperation = CompilerContext.SpawnIntermediateNode<UK2Node_CallFunction>(this, SourceGraph);
 	CallOperation->FunctionReference.SetExternalMember( GET_FUNCTION_NAME_CHECKED(UDataprepActionAsset, Execute), UDataprepActionAsset::StaticClass() );
@@ -196,12 +185,12 @@ void UK2Node_DataprepAction::ExpandNode(FKismetCompilerContext& CompilerContext,
 	CompilerContext.MovePinLinksToIntermediate( *CallOperation->GetThenPin(), GetOutExecutionPin() );
 }
 
-UEdGraphPin& UK2Node_DataprepAction::GetOutExecutionPin() const
+UEdGraphPin& UDEPRECATED_K2Node_DataprepAction::GetOutExecutionPin() const
 {
 	return *FindPinChecked(UEdGraphSchema_K2::PN_Then, EGPD_Output);
 }
 
-UEdGraphPin& UK2Node_DataprepAction::GetInObjectsPin() const
+UEdGraphPin& UDEPRECATED_K2Node_DataprepAction::GetInObjectsPin() const
 {
 	return *FindPinChecked( InObjectsPinName, EGPD_Input);
 }
diff --git a/Engine/Plugins/Enterprise/DataprepEditor/Source/DataprepEditor/Private/BlueprintNodes/K2Node_DataprepAction.h b/Engine/Plugins/Enterprise/DataprepEditor/Source/DataprepEditor/Private/BlueprintNodes/K2Node_DataprepAction.h
--- a/Engine/Plugins/Enterprise/DataprepEditor/Source/DataprepEditor/Private/BlueprintNodes/K2Node_DataprepAction.h
+++ b/Engine/Plugins/Enterprise/DataprepEditor/Source/DataprepEditor/Private/BlueprintNodes/K2Node_DataprepAction.h
@@ -15,13 +15,13 @@ class FBlueprintActionDatabaseRegistrar;
 class UDataprepActionAsset;
 class UEdGraphPin;
 
-UCLASS(MinimalAPI)
-class UK2Node_DataprepAction : public UK2Node_DataprepActionCore
+UCLASS(MinimalAPI, Deprecated, meta = (DeprecationMessage = "No use of Blueprint with Dataprep."))
+class UDEPRECATED_K2Node_DataprepAction : public UK2Node_DataprepActionCore
 {
 	GENERATED_BODY()
 
 public:
-	UK2Node_DataprepAction();
+	UDEPRECATED_K2Node_DataprepAction();
 
 	// Begin EdGraphNode interface
 	virtual void AllocateDefaultPins() override;
@@ -30,8 +30,6 @@ public:
 	virtual FText GetNodeTitle(ENodeTitleType::Type TitleType) const override;
 	virtual void OnRenameNode(const FString& NewName) override;
 	virtual void DestroyNode() override;
-	virtual void NodeConnectionListChanged() override;
-	virtual TSharedPtr<SGraphNode> CreateVisualWidget() override;
 	TSharedPtr<class INameValidatorInterface> MakeNameValidator() const override;
 	// End EdGraphNode interface
 
diff --git a/Engine/Plugins/Enterprise/DataprepEditor/Source/DataprepEditor/Private/DataprepDetailsTab.cpp b/Engine/Plugins/Enterprise/DataprepEditor/Source/DataprepEditor/Private/DataprepDetailsTab.cpp
--- a/Engine/Plugins/Enterprise/DataprepEditor/Source/DataprepEditor/Private/DataprepDetailsTab.cpp
+++ b/Engine/Plugins/Enterprise/DataprepEditor/Source/DataprepEditor/Private/DataprepDetailsTab.cpp
@@ -8,7 +8,6 @@
 #include "DataprepEditorActions.h"
 #include "DataprepEditorModule.h"
 #include "DataprepEditorStyle.h"
-#include "DataprepRecipe.h"
 #include "Widgets/SAssetsPreviewWidget.h"
 
 #include "AssetRegistryModule.h"
diff --git a/Engine/Plugins/Enterprise/DataprepEditor/Source/DataprepEditor/Private/DataprepEditor.cpp b/Engine/Plugins/Enterprise/DataprepEditor/Source/DataprepEditor/Private/DataprepEditor.cpp
--- a/Engine/Plugins/Enterprise/DataprepEditor/Source/DataprepEditor/Private/DataprepEditor.cpp
+++ b/Engine/Plugins/Enterprise/DataprepEditor/Source/DataprepEditor/Private/DataprepEditor.cpp
@@ -13,7 +13,6 @@
 #include "DataprepEditorLogCategory.h"
 #include "DataprepEditorModule.h"
 #include "DataprepEditorStyle.h"
-#include "DataprepRecipe.h"
 #include "PreviewSystem/DataprepPreviewAssetColumn.h"
 #include "PreviewSystem/DataprepPreviewSceneOutlinerColumn.h"
 #include "PreviewSystem/DataprepPreviewSystem.h"
@@ -64,21 +63,6 @@
 #include "Widgets/Input/STextComboBox.h"
 #include "Widgets/Layout/SBorder.h"
 
-// Temp code for the nodes development
-#include "Engine/Blueprint.h"
-#include "Kismet2/BlueprintEditorUtils.h"
-#include "K2Node_CustomEvent.h"
-#include "DataprepRecipe.h"
-#include "Engine/BlueprintGeneratedClass.h"
-#include "Kismet2/KismetEditorUtilities.h"
-#include "PropertyEditorModule.h"
-#include "IDetailsView.h"
-#include "DataprepAsset.h"
-#include "IStructureDetailsView.h"
-
-const FName FDataprepEditor::PipelineGraphTabId(TEXT("DataprepEditor_Pipeline_Graph"));
-// end of temp code for nodes development
-
 #define LOCTEXT_NAMESPACE "DataprepEditor"
 
 extern const FName DataprepEditorAppIdentifier;
@@ -315,15 +299,8 @@ void FDataprepEditor::RegisterTabSpawners(const TSharedRef<class FTabManager>& I
 			.SetGroup(WorkspaceMenuCategoryRef)
 			.SetIcon( FSlateIcon(FEditorStyle::GetStyleSetName(), "Kismet.Tabs.Palette"));
 
-		// Temp code for the nodes development
-		InTabManager->RegisterTabSpawner(PipelineGraphTabId, FOnSpawnTab::CreateSP(this, &FDataprepEditor::SpawnTabPipelineGraph))
-			.SetDisplayName(LOCTEXT("PipelineGraphTab", "Pipeline Graph"))
-			.SetGroup(WorkspaceMenuCategoryRef)
-			.SetIcon(FSlateIcon(FEditorStyle::GetStyleSetName(), "GraphEditor.EventGraph_16x"));
-		// end of temp code for nodes development
-
 		InTabManager->RegisterTabSpawner(DataprepGraphEditorTabId, FOnSpawnTab::CreateSP(this, &FDataprepEditor::SpawnTabGraphEditor))
-			.SetDisplayName(LOCTEXT("GraphEditorTab", "Simplified Graph"))
+			.SetDisplayName(LOCTEXT("GraphEditorTab", "Recipe Graph"))
 			.SetGroup(WorkspaceMenuCategoryRef)
 			.SetIcon(FSlateIcon(FEditorStyle::GetStyleSetName(), "GraphEditor.EventGraph_16x"));
 		
@@ -400,7 +377,7 @@ const FString& FDataprepEditor::GetRootPackagePath()
 	return RootPackagePath;
 }
 
-void FDataprepEditor::InitDataprepEditor(const EToolkitMode::Type Mode, const TSharedPtr<IToolkitHost>& InitToolkitHost, UDataprepAssetInterface* InDataprepAssetInterface, UObject* Blueprint)
+void FDataprepEditor::InitDataprepEditor(const EToolkitMode::Type Mode, const TSharedPtr<IToolkitHost>& InitToolkitHost, UDataprepAssetInterface* InDataprepAssetInterface)
 {
 	DataprepAssetInterfacePtr = TWeakObjectPtr<UDataprepAssetInterface>(InDataprepAssetInterface);
 	check( DataprepAssetInterfacePtr.IsValid() );
@@ -440,22 +417,6 @@ void FDataprepEditor::InitDataprepEditor(const EToolkitMode::Type Mode, const TS
 	TempDir = FPaths::Combine( GetRootTemporaryDir(), FString::FromInt( FPlatformProcess::GetCurrentProcessId() ), SessionID);
 	IFileManager::Get().MakeDirectory(*TempDir);
 
-#ifndef NO_BLUEPRINT
-	// Temp code for the nodes development
-	if(Blueprint != nullptr)
-	{
-		DataprepRecipeBPPtr = Cast<UBlueprint>(Blueprint);
-		check( DataprepRecipeBPPtr.IsValid() );
-
-		// Necessary step to regenerate blueprint generated class
-		// Note that this compilation will always succeed as Dataprep node does not have real body
-		//{
-		//	FKismetEditorUtilities::CompileBlueprint( DataprepRecipeBPPtr.Get(), EBlueprintCompileOptions::None, nullptr );
-		//}
-	}
-	// End temp code for the nodes development
-#endif
-
 	GEditor->RegisterForUndo(this);
 
 	// Register our commands. This will only register them if not previously registered
@@ -464,9 +425,8 @@ void FDataprepEditor::InitDataprepEditor(const EToolkitMode::Type Mode, const TS
 	BindCommands();
 
 	CreateTabs();
-
 	
-	const TSharedRef<FTabManager::FLayout> Layout = DataprepRecipeBPPtr.IsValid() ? CreateDataprepLayout() : CreateDataprepInstanceLayout();
+	const TSharedRef<FTabManager::FLayout> Layout = bIsDataprepInstance ? CreateDataprepInstanceLayout() : CreateDataprepLayout();
 
 	const bool bCreateDefaultStandaloneMenu = true;
 	const bool bCreateDefaultToolbar = true;
@@ -487,22 +447,6 @@ void FDataprepEditor::BindCommands()
 
 	const TSharedRef<FUICommandList>& UICommandList = GetToolkitCommands();
 
-	//UICommandList->MapAction(FGenericCommands::Get().Delete,
-	//	FExecuteAction::CreateSP(this, &FDataprepEditor::DeleteSelected),
-	//	FCanExecuteAction::CreateSP(this, &FDataprepEditor::CanDeleteSelected));
-
-	//UICommandList->MapAction(FGenericCommands::Get().Undo,
-	//	FExecuteAction::CreateSP(this, &FDataprepEditor::UndoAction));
-
-	//UICommandList->MapAction(FGenericCommands::Get().Redo,
-	//	FExecuteAction::CreateSP(this, &FDataprepEditor::RedoAction));
-
-	// Temp code for the nodes development
-	UICommandList->MapAction(
-		Commands.CompileGraph,
-		FExecuteAction::CreateSP(this, &FDataprepEditor::OnCompile));
-	// end of temp code for nodes development
-
 	UICommandList->MapAction(
 		Commands.SaveScene,
 		FExecuteAction::CreateSP(this, &FDataprepEditor::OnSaveScene));
@@ -868,7 +812,9 @@ void FDataprepEditor::CreateTabs()
 		}
 	);
 
-	DataprepAssetView = SNew( SDataprepAssetView, DataprepAssetInterfacePtr.Get(), PipelineEditorCommands );
+	CreateGraphEditor();
+
+	DataprepAssetView = SNew( SDataprepAssetView, DataprepAssetInterfacePtr.Get() );
 
 	CreateScenePreviewTab();
 
@@ -878,37 +824,7 @@ void FDataprepEditor::CreateTabs()
 
 	// Create Details Panel
 	CreateDetailsViews();
-
-	// Temp code for the nodes development
-	if(DataprepRecipeBPPtr.IsValid())
-	{
-		// Create Pipeline Editor
-		CreatePipelineEditor();
-	}
-	// end of temp code for nodes development
-
-	CreateGraphEditor();
-
-}
-
-// Temp code for the nodes development
-TSharedRef<SDockTab> FDataprepEditor::SpawnTabPipelineGraph(const FSpawnTabArgs & Args)
-{
-	check(Args.GetTabId() == PipelineGraphTabId);
-
-	if(!bIsDataprepInstance)
-	{
-		return SNew(SDockTab)
-			//.Icon(FDataprepEditorStyle::Get()->GetBrush("DataprepEditor.Tabs.Pipeline"))
-			.Label(LOCTEXT("DataprepEditor_PipelineTab_Title", "Pipeline"))
-			[
-				DataprepRecipeBPPtr.IsValid() ? PipelineView.ToSharedRef() : SNullWidget::NullWidget
-			];
-	}
-
-	return SNew(SDockTab);
 }
-// end of temp code for nodes development
 
 TSharedRef<SDockTab> FDataprepEditor::SpawnTabScenePreview(const FSpawnTabArgs & Args)
 {
@@ -1010,7 +926,7 @@ TSharedRef<SDockTab> FDataprepEditor::SpawnTabGraphEditor(const FSpawnTabArgs &
 	{
 		return SNew(SDockTab)
 			//.Icon(FDataprepEditorStyle::Get()->GetBrush("DataprepEditor.Tabs.Pipeline"))
-			.Label(LOCTEXT("DataprepEditor_GraphEditorTab_Title", "Simplified Graph"))
+			.Label(LOCTEXT("DataprepEditor_GraphEditorTab_Title", "Recipe Graph"))
 			[
 				GraphEditor.ToSharedRef()
 			];
@@ -1021,7 +937,7 @@ TSharedRef<SDockTab> FDataprepEditor::SpawnTabGraphEditor(const FSpawnTabArgs &
 
 TSharedRef<FTabManager::FLayout> FDataprepEditor::CreateDataprepLayout()
 {
-	return FTabManager::NewLayout("Standalone_DataprepEditor_Layout_v0.7")
+	return FTabManager::NewLayout("Standalone_DataprepEditor_Layout_v0.8")
 		->AddArea
 		(
 			FTabManager::NewPrimaryArea()->SetOrientation(Orient_Vertical)
@@ -1074,15 +990,6 @@ TSharedRef<FTabManager::FLayout> FDataprepEditor::CreateDataprepLayout()
 							->AddTab(PaletteTabId, ETabState::OpenedTab)
 							->SetHideTabWell( true )
 						)
-						// Temp code for the nodes development
-						->Split
-						(
-							FTabManager::NewStack()
-							->SetSizeCoefficient(0.85f)
-							->AddTab(PipelineGraphTabId, ETabState::OpenedTab)
-							->SetHideTabWell( true )
-						)
-						// end of temp code for nodes development
 						->Split
 						(
 							FTabManager::NewStack()
diff --git a/Engine/Plugins/Enterprise/DataprepEditor/Source/DataprepEditor/Private/DataprepEditor.h b/Engine/Plugins/Enterprise/DataprepEditor/Source/DataprepEditor/Private/DataprepEditor.h
--- a/Engine/Plugins/Enterprise/DataprepEditor/Source/DataprepEditor/Private/DataprepEditor.h
+++ b/Engine/Plugins/Enterprise/DataprepEditor/Source/DataprepEditor/Private/DataprepEditor.h
@@ -75,7 +75,7 @@ public:
 	virtual void RegisterTabSpawners(const TSharedRef<FTabManager>& TabManager) override;
 	virtual void UnregisterTabSpawners(const TSharedRef<FTabManager>& TabManager) override;
 
-	void InitDataprepEditor(const EToolkitMode::Type Mode, const TSharedPtr< IToolkitHost >& InitToolkitHost, UDataprepAssetInterface* InDataprepAssetInterface, UObject* Blueprint = nullptr);
+	void InitDataprepEditor(const EToolkitMode::Type Mode, const TSharedPtr< IToolkitHost >& InitToolkitHost, UDataprepAssetInterface* InDataprepAssetInterface);
 
 	/** IToolkit interface */
 	virtual FName GetToolkitFName() const override;
@@ -205,6 +205,8 @@ private:
 	virtual void PostUndo( bool bSuccess ) override;
 	virtual void PostRedo( bool bSuccess ) override;
 
+	void CreateGraphEditor();
+
 private:
 	bool bWorldBuilt;
 	bool bIsFirstRun;
@@ -221,13 +223,11 @@ private:
 	TSharedPtr<SWidget> ScenePreviewView;
 	TSharedPtr<SGraphNodeDetailsWidget> DetailsView;
 	TSharedPtr<class SDataprepAssetView > DataprepAssetView;
-	TSharedPtr<SGraphEditor> PipelineView;
 	TSharedPtr<SDataprepGraphEditor> GraphEditor;
 
 	TSharedPtr<class ICustomSceneOutliner> SceneOutliner;
 
 	/** Command list for the pipeline editor */
-	TSharedPtr<FUICommandList> PipelineEditorCommands;
 	TSharedPtr<FUICommandList> GraphEditorCommands;
 	bool bIsActionMenuContextSensitive;
 	bool bSaveIntermediateBuildProducts;
@@ -277,6 +277,8 @@ private:
 	/** Steps Preview System */
 	TSharedRef<FDataprepPreviewSystem> PreviewSystem;
 
+	TSharedPtr<FDataprepActionContext> ActionsContext;
+
 	/**	The tab ids for all the tabs used */
 	static const FName ScenePreviewTabId;
 	static const FName AssetPreviewTabId;
@@ -286,87 +288,4 @@ private:
 	static const FName SceneViewportTabId;
 	static const FName DataprepStatisticsTabId;
 	static const FName DataprepGraphEditorTabId;
-
-//Temp Code to allow us to work on the nodes while we don't have our own graph.
-public:
-	UBlueprint* GetDataprepBlueprint()
-	{
-		return DataprepRecipeBPPtr.IsValid() ? DataprepRecipeBPPtr.Get() : nullptr;
-	}
-
-private:
-
-	TSharedRef<SDockTab> SpawnTabPipelineGraph(const FSpawnTabArgs& Args);
-
-	void CreatePipelineEditor();
-
-	void CreateGraphEditor();
-
-	
-	/** Called to create context menu when right-clicking on graph */
-	FActionMenuContent OnCreatePipelineActionMenu(UEdGraph* InGraph, const FVector2D& InNodePosition, const TArray<UEdGraphPin*>& InDraggedPins, bool bAutoExpand, SGraphEditor::FActionMenuClosed InOnMenuClosed);
-
-	/** Called when the Pipeline Blueprint has changed */
-	void OnPipelineChanged(UBlueprint* InBlueprint);
-
-	/** Called when the Pipeline Blueprint has been compiled */
-	void OnPipelineCompiled(UBlueprint* InBlueprint);
-
-	/** Callback when a token is clicked on in the compiler results log */
-	void OnLogTokenClicked(const TSharedRef<IMessageToken>& Token);
-
-	/** Callback when properties have finished being handled */
-	//virtual void OnFinishedChangingProperties(const FPropertyChangedEvent& PropertyChangedEvent);
-
-	/** Called when Compile button is clicked */
-	void OnCompile();
-
-	/** Returns the current status icon for the blueprint being edited */
-	FSlateIcon GetPipelineCompileButtonImage() const;
-	FText GetPipelineCompileButtonTooltip() const;
-
-	TSet<UObject*> GetSelectedPipelineNodes() const;
-	UEdGraph* GetPipelineGraph() const;
-	bool IsPipelineEditable() const;
-
-	void OnRenameNode();
-	bool CanRenameNode() const;
-
-	void SelectAllNodes();
-	bool CanSelectAllNodes() const;
-
-	void DeleteSelectedPipelineNodes();
-	bool CanDeletePipelineNodes() const;
-
-	void CopySelectedNodes();
-	bool CanCopyNodes() const;
-
-	void CutSelectedNodes();
-	bool CanCutNodes() const;
-
-	void PasteNodes();
-	bool CanPasteNodes() const;
-
-	void DuplicateNodes();
-	bool CanDuplicateNodes() const;
-
-	void OnCreateComment();
-
-	void DeleteSelectedDuplicatableNodes();
-	void PasteNodesHere(class UEdGraph* DestinationGraph, const FVector2D& GraphLocation);
-
-	bool OnNodeVerifyTitleCommit(const FText& NewText, UEdGraphNode* NodeBeingChanged, FText& OutErrorMessage);
-	void OnNodeTitleCommitted(const FText& NewText, ETextCommit::Type CommitInfo, UEdGraphNode* NodeBeingChanged);
-
-	// Start of temp code
-	TWeakObjectPtr<UBlueprint> DataprepRecipeBPPtr;
-
-	/** Compiler results log, with the log listing that it reflects */
-	TSharedPtr<SWidget> CompilerResults;
-	TSharedPtr<IMessageLogListing> CompilerResultsListing;
-
-	TSharedPtr<FDataprepActionContext> ActionsContext;
-
-	static const FName PipelineGraphTabId;
-	// End of temp code
 };
diff --git a/Engine/Plugins/Enterprise/DataprepEditor/Source/DataprepEditor/Private/DataprepEditorModule.cpp b/Engine/Plugins/Enterprise/DataprepEditor/Source/DataprepEditor/Private/DataprepEditorModule.cpp
--- a/Engine/Plugins/Enterprise/DataprepEditor/Source/DataprepEditor/Private/DataprepEditorModule.cpp
+++ b/Engine/Plugins/Enterprise/DataprepEditor/Source/DataprepEditor/Private/DataprepEditorModule.cpp
@@ -2,7 +2,6 @@
 
 #include "DataprepEditorModule.h"
 
-#include "AssetTypeActions_Dataprep.h"
 #include "AssetTypeActions_DataprepAsset.h"
 #include "AssetTypeActions_DataprepAssetInterface.h"
 #include "DataprepAssetProducers.h"
@@ -57,10 +56,6 @@ public:
 		AssetTools.RegisterAssetTypeActions(DataprepAssetInterfaceTypeAction.ToSharedRef());
 		AssetTypeActionsArray.Add(DataprepAssetInterfaceTypeAction);
 
-		TSharedPtr<FAssetTypeActions_Dataprep> DataprepRecipeTypeAction = MakeShareable(new FAssetTypeActions_Dataprep);
-		AssetTools.RegisterAssetTypeActions(DataprepRecipeTypeAction.ToSharedRef());
-		AssetTypeActionsArray.Add(DataprepRecipeTypeAction);
-
 		TSharedPtr<FAssetTypeActions_DataprepAsset> DataprepAssetTypeAction = MakeShareable(new FAssetTypeActions_DataprepAsset);
 		AssetTools.RegisterAssetTypeActions(DataprepAssetTypeAction.ToSharedRef());
 		AssetTypeActionsArray.Add(DataprepAssetTypeAction);
@@ -115,8 +110,7 @@ public:
 
 	virtual TSharedRef<SWidget> CreateDataprepProducersWidget(UDataprepAssetProducers* AssetProducers) override
 	{
-		TSharedPtr<FUICommandList> CommandList = MakeShareable( new FUICommandList );
-		return AssetProducers ? SNew(SDataprepProducersWidget, AssetProducers, CommandList) : SNullWidget::NullWidget;
+		return AssetProducers ? SNew(SDataprepProducersWidget, AssetProducers) : SNullWidget::NullWidget;
 	}
 
 	virtual TSharedRef<SWidget> CreateDataprepDetailsView(UObject* ObjectToDetail) override
diff --git a/Engine/Plugins/Enterprise/DataprepEditor/Source/DataprepEditor/Private/DataprepEditorUtils.cpp b/Engine/Plugins/Enterprise/DataprepEditor/Source/DataprepEditor/Private/DataprepEditorUtils.cpp
--- a/Engine/Plugins/Enterprise/DataprepEditor/Source/DataprepEditor/Private/DataprepEditorUtils.cpp
+++ b/Engine/Plugins/Enterprise/DataprepEditor/Source/DataprepEditor/Private/DataprepEditorUtils.cpp
@@ -27,41 +27,6 @@ bool FDataprepParametrizationActionData::IsValid() const
 	return DataprepAsset && Object && PropertyChain.Num() > 0;
 }
 
-#ifndef NO_BLUEPRINT
-void FDataprepEditorUtils::NotifySystemOfChangeInPipeline(UObject* SourceObject)
-{
-	UBlueprint* Blueprint = nullptr;
-	UDataprepAsset* DataprepAsset = nullptr;
-	UObject* Object = SourceObject;
-	while ( Object )
-	{
-		UClass* Class = Object->GetClass();
-		if ( Class->IsChildOf<UBlueprint>() )
-		{
-			Blueprint = static_cast<UBlueprint*>( Object );
-		}
-		else if ( Class == UDataprepAsset::StaticClass() )
-		{
-			DataprepAsset = static_cast<UDataprepAsset*>( Object );
-			break;
-		}
-		Object = Object->GetOuter();
-	}
-
-	if ( DataprepAsset )
-	{
-		if(!DataprepAsset->HasAnyFlags(RF_NeedLoad | RF_NeedPostLoad | RF_NeedPostLoadSubobjects))
-		{
-			DataprepAsset->GetRecipeBP()->OnChanged().Broadcast(DataprepAsset->GetRecipeBP());
-		}
-	}
-	else if ( Blueprint )
-	{
-		FBlueprintEditorUtils::MarkBlueprintAsModified( Blueprint );
-	}
-}
-#endif
-
 void FDataprepEditorUtils::PopulateMenuForParameterization(FMenuBuilder& MenuBuilder, UDataprepAsset& DataprepAsset, UDataprepParameterizableObject& Object, const TArray<FDataprepPropertyLink>& PropertyChain)
 {
 	TSharedRef<FDataprepParametrizationActionData> ActionData = MakeShared<FDataprepParametrizationActionData>( DataprepAsset, Object, PropertyChain );
diff --git a/Engine/Plugins/Enterprise/DataprepEditor/Source/DataprepEditor/Private/DataprepEditorUtils.h b/Engine/Plugins/Enterprise/DataprepEditor/Source/DataprepEditor/Private/DataprepEditorUtils.h
--- a/Engine/Plugins/Enterprise/DataprepEditor/Source/DataprepEditor/Private/DataprepEditorUtils.h
+++ b/Engine/Plugins/Enterprise/DataprepEditor/Source/DataprepEditor/Private/DataprepEditorUtils.h
@@ -33,13 +33,6 @@ struct FDataprepParametrizationActionData : public FGCObject
 class FDataprepEditorUtils
 {
 public:
-#ifndef NO_BLUEPRINT
-	/**
-	 * Find the the owning dataprep asset of the source object and send a notification to the dataprep editor so it can react when its pipeline is modified
-	 */
-	static void NotifySystemOfChangeInPipeline(UObject* SourceObject);
-#endif
-
 	/**
 	 * Populate a menu builder with the section made for the parameterization
 	 * @param DataprepAsset the asset that own the object
diff --git a/Engine/Plugins/Enterprise/DataprepEditor/Source/DataprepEditor/Private/DataprepPipelineEditor.cpp b/Engine/Plugins/Enterprise/DataprepEditor/Source/DataprepEditor/Private/DataprepPipelineEditor.cpp
deleted file mode 100644
diff --git a/Engine/Plugins/Enterprise/DataprepEditor/Source/DataprepEditor/Private/SchemaActions/DataprepDragDropOp.cpp b/Engine/Plugins/Enterprise/DataprepEditor/Source/DataprepEditor/Private/SchemaActions/DataprepDragDropOp.cpp
--- a/Engine/Plugins/Enterprise/DataprepEditor/Source/DataprepEditor/Private/SchemaActions/DataprepDragDropOp.cpp
+++ b/Engine/Plugins/Enterprise/DataprepEditor/Source/DataprepEditor/Private/SchemaActions/DataprepDragDropOp.cpp
@@ -154,35 +154,6 @@ void FDataprepDragDropOp::HoverTargetChanged()
 	SetSimpleFeedbackMessage( Symbol, FLinearColor::White, DrapDropText );
 }
 
-FReply FDataprepDragDropOp::DroppedOnPanel(const TSharedRef<SWidget>& Panel, FVector2D ScreenPosition, FVector2D GraphPosition, UEdGraph& Graph)
-{
-	if ( bDropTargetValid )
-	{
-		if ( DataprepPreDropConfirmation.IsBound() )
-		{
-			TFunction<void ()> OnConfirmation ( [Operation = StaticCastSharedRef<FDataprepDragDropOp>( AsShared() ), Panel, ScreenPosition, GraphPosition, GraphPtr = TWeakObjectPtr<UEdGraph>(&Graph)] ()
-				{
-					UEdGraph* Graph = GraphPtr.Get();
-					if ( Graph )
-					{
-						Operation->DoDropOnPanel(Panel, ScreenPosition, GraphPosition, *Graph);
-					}
-				} );
-
-			DataprepPreDropConfirmation.Execute( FDataprepSchemaActionContext(), OnConfirmation);
-		
-			return FReply::Handled();
-		}
-		else
-		{
-			DoDropOnPanel( Panel, ScreenPosition, GraphPosition, Graph );
-			return FReply::Handled();
-		}
-	}
-
-	return FReply::Unhandled();
-}
-
 void FDataprepDragDropOp::OnDragged(const FDragDropEvent& DragDropEvent)
 {
 	FVector2D TargetPosition = DragDropEvent.GetScreenSpacePosition();
@@ -372,35 +343,6 @@ bool FDataprepDragDropOp::DoDropOnDataprepActionContext(const FDataprepSchemaAct
 	return false;
 }
 
-void FDataprepDragDropOp::DoDropOnPanel(const TSharedRef<SWidget>& Panel, FVector2D ScreenPosition, FVector2D GraphPosition, UEdGraph& Graph)
-{
-	if (UEdGraph* EdGraph = GetHoveredGraph())
-	{
-		if(Cast<UDataprepGraph>(EdGraph) != nullptr)
-		{
-			return;
-		}
-
-		FScopedTransaction Transaction( LOCTEXT("AddNode", "Add Dataprep Action Node") );
-		UK2Node_DataprepAction* DataprepActionNode = DataprepSchemaActionUtils::SpawnEdGraphNode< UK2Node_DataprepAction >( Graph, GraphPosition );
-		check( DataprepActionNode );
-		DataprepActionNode->CreateDataprepActionAsset();
-		DataprepActionNode->AutowireNewNode( GetHoveredPin() );
-
-		FDataprepSchemaActionContext Context;
-		Context.DataprepActionPtr = DataprepActionNode->GetDataprepAction();
-		if ( !DoDropOnDataprepActionContext( Context ) )
-		{
-			Transaction.Cancel();
-		}
-
-		if ( UBlueprint* Blueprint = FBlueprintEditorUtils::FindBlueprintForGraphChecked( EdGraph ) )
-		{
-			FBlueprintEditorUtils::MarkBlueprintAsModified( Blueprint );
-		}
-	}
-}
-
 FReply FDataprepDragDropOp::DoDropOnActionStep(UDataprepGraphActionStepNode* TargetActionStepNode)
 {
 	FModifierKeysState ModifierKeyState = FSlateApplication::Get().GetModifierKeys();
diff --git a/Engine/Plugins/Enterprise/DataprepEditor/Source/DataprepEditor/Private/SchemaActions/DataprepDragDropOp.h b/Engine/Plugins/Enterprise/DataprepEditor/Source/DataprepEditor/Private/SchemaActions/DataprepDragDropOp.h
--- a/Engine/Plugins/Enterprise/DataprepEditor/Source/DataprepEditor/Private/SchemaActions/DataprepDragDropOp.h
+++ b/Engine/Plugins/Enterprise/DataprepEditor/Source/DataprepEditor/Private/SchemaActions/DataprepDragDropOp.h
@@ -46,7 +46,6 @@ public:
 	// FGraphEditorDragDropAction Interface
 	virtual void HoverTargetChanged() override;
 	virtual FReply DroppedOnNode(FVector2D ScreenPosition, FVector2D GraphPosition) override;
-	virtual FReply DroppedOnPanel(const TSharedRef<SWidget>& Panel, FVector2D ScreenPosition, FVector2D GraphPosition, UEdGraph& Graph) override;
 	virtual void OnDragged (const class FDragDropEvent& DragDropEvent ) override;
 	virtual EVisibility GetIconVisible() const;
 	virtual EVisibility GetErrorIconVisible() const;
@@ -78,7 +77,6 @@ protected:
 	typedef FGraphEditorDragDropAction Super;
 
 	bool DoDropOnDataprepActionContext(const FDataprepSchemaActionContext& Context);
-	void DoDropOnPanel(const TSharedRef<SWidget>& Panel, FVector2D ScreenPosition, FVector2D GraphPosition, UEdGraph& Graph);
 
 	/** Executes drop on existing action step */
 	FReply DoDropOnActionStep(UDataprepGraphActionStepNode* TargetActionStepNode);
diff --git a/Engine/Plugins/Enterprise/DataprepEditor/Source/DataprepEditor/Private/Widgets/DataprepAssetView.cpp b/Engine/Plugins/Enterprise/DataprepEditor/Source/DataprepEditor/Private/Widgets/DataprepAssetView.cpp
--- a/Engine/Plugins/Enterprise/DataprepEditor/Source/DataprepEditor/Private/Widgets/DataprepAssetView.cpp
+++ b/Engine/Plugins/Enterprise/DataprepEditor/Source/DataprepEditor/Private/Widgets/DataprepAssetView.cpp
@@ -297,7 +297,7 @@ TSharedRef<ITableRow> SDataprepAssetView::OnGenerateRowForCategoryTree( TSharedR
 	{
 		case EDataprepCategory::Producers:
 		{
-			ProducersWidget = SNew( SDataprepProducersWidget, DataprepAssetInterfacePtr->GetProducers(), CommandList )
+			ProducersWidget = SNew( SDataprepProducersWidget, DataprepAssetInterfacePtr->GetProducers() )
 				.ColumnSizeData( ColumnSizeData );
 
 			TSharedPtr< SWidget > ProducerWrapper = SNew( SHorizontalBox )
@@ -393,12 +393,11 @@ TSharedRef<ITableRow> SDataprepAssetView::OnGenerateRowForCategoryTree( TSharedR
 	return Row.ToSharedRef();
 }
 
-void SDataprepAssetView::Construct( const FArguments& InArgs, UDataprepAssetInterface* InDataprepAssetPtr, TSharedPtr<FUICommandList>& InCommandList )
+void SDataprepAssetView::Construct( const FArguments& InArgs, UDataprepAssetInterface* InDataprepAssetPtr )
 {
 	check( InDataprepAssetPtr );
 
 	DataprepAssetInterfacePtr = InDataprepAssetPtr;
-	CommandList = InCommandList;
 
 	DataprepAssetInterfacePtr->GetOnChanged().AddRaw( this, &SDataprepAssetView::OnDataprepAssetChanged );
 
diff --git a/Engine/Plugins/Enterprise/DataprepEditor/Source/DataprepEditor/Private/Widgets/DataprepAssetView.h b/Engine/Plugins/Enterprise/DataprepEditor/Source/DataprepEditor/Private/Widgets/DataprepAssetView.h
--- a/Engine/Plugins/Enterprise/DataprepEditor/Source/DataprepEditor/Private/Widgets/DataprepAssetView.h
+++ b/Engine/Plugins/Enterprise/DataprepEditor/Source/DataprepEditor/Private/Widgets/DataprepAssetView.h
@@ -29,7 +29,7 @@ public:
 	SLATE_BEGIN_ARGS(SDataprepAssetView) {}
 	SLATE_END_ARGS()
 
-	void Construct(const FArguments& InArgs, UDataprepAssetInterface* InDataprepAssetPtr, TSharedPtr<FUICommandList>& InCommandList);
+	void Construct(const FArguments& InArgs, UDataprepAssetInterface* InDataprepAssetPtr);
 
 	~SDataprepAssetView();
 
@@ -47,7 +47,6 @@ private:
 	void OnGetChildrenForCategoryTree( TSharedRef<EDataprepCategory> InTreeNode, TArray< TSharedRef<EDataprepCategory> >& OutChildren ) {}
 
 	TArray<TSharedRef< EDataprepCategory >> Categories;
-	TSharedPtr<FUICommandList> CommandList;
 
 private:
 	TWeakObjectPtr<UDataprepAssetInterface> DataprepAssetInterfacePtr;
diff --git a/Engine/Plugins/Enterprise/DataprepEditor/Source/DataprepEditor/Private/Widgets/DataprepGraph/SDataprepActionBlock.cpp b/Engine/Plugins/Enterprise/DataprepEditor/Source/DataprepEditor/Private/Widgets/DataprepGraph/SDataprepActionBlock.cpp
--- a/Engine/Plugins/Enterprise/DataprepEditor/Source/DataprepEditor/Private/Widgets/DataprepGraph/SDataprepActionBlock.cpp
+++ b/Engine/Plugins/Enterprise/DataprepEditor/Source/DataprepEditor/Private/Widgets/DataprepGraph/SDataprepActionBlock.cpp
@@ -35,93 +35,6 @@ void SDataprepActionBlock::Construct(const FArguments& InArgs, const TSharedRef<
 {
 	DataprepActionContext = InDataprepActionContext;
 
-	const float DefaultPadding = FDataprepEditorStyle::GetFloat( "DataprepAction.Padding" );
-	const FLinearColor OutlineColor = GetOutlineColor().GetSpecifiedColor();
-#ifndef NO_BLUEPRINT
-	bIsSimplifiedGraph = InArgs._IsSimplified;
-
-	if(bIsSimplifiedGraph)
-	{
-		ConstructForSimplified();
-		return;
-	}
-#endif
-
-	ChildSlot
-	[
-		SNew( SConstraintCanvas )
-
-		// The outline. This is done by a background image
-		+ SConstraintCanvas::Slot()
-		.Anchors( FAnchors( 0.f, 0.f, 1.f, 1.f ) )
-		.Offset( FMargin() )
-		[
-			SNew( SColorBlock )
-			.Color( OutlineColor )
-		]
-
-		+ SConstraintCanvas::Slot()
-		.Anchors( FAnchors( 0.f, 0.f, 1.f, 1.f ) )
-		.Offset( FMargin() )
-		.AutoSize( true )
-		[
-			SNew( SVerticalBox )
-
-			//The title of the block
-			+ SVerticalBox::Slot()
-			.AutoHeight()
-			.Padding( FMargin( DefaultPadding, DefaultPadding, DefaultPadding, 0.f ) )
-			[
-				SNew( SConstraintCanvas )
-
-				// The background of the title
-				+ SConstraintCanvas::Slot()
-				.Anchors( FAnchors( 0.f, 0.f, 1.f, 1.f ) )
-				.Offset( FMargin() )
-				[
-					GetTitleBackgroundWidget()
-				]
-
-				// The title of the block
-				+ SConstraintCanvas::Slot()
-				.Anchors( FAnchors( 0.5f, 0.5f, 0.5f, 0.5f ) )
-				.Offset( FMargin() )
-				.AutoSize( true )
-				[
-					GetTitleWidget()
-				]
-			]
-
-			// The content zone of the action block
-			+ SVerticalBox::Slot()
-			.AutoHeight()
-			.Padding( FMargin( DefaultPadding ) )
-			[
-				SNew( SConstraintCanvas )
-
-				// The background of the content zone
-				+ SConstraintCanvas::Slot()
-				.Anchors( FAnchors( 0.f, 0.f, 1.f, 1.f ) )
-				.Offset( FMargin() )
-				[
-					GetContentBackgroundWidget()
-				]
-
-				// The content of the content zone
-				+ SConstraintCanvas::Slot()
-				.Anchors(FAnchors( 0.f, 0.f, 1.f, 1.f ) )
-				.Offset( FMargin() )
-				.AutoSize( true )
-				[
-					GetContentWidget()
-				]
-			]
-		]
-	];
-}
-
-void SDataprepActionBlock::ConstructForSimplified()
-{
 	const FLinearColor OutlineColor = GetOutlineColor().GetSpecifiedColor();
 
 	ChildSlot
@@ -202,174 +115,6 @@ TSharedRef<SWidget> SDataprepActionBlock::GetBlockTitleWidget()
 	];
 }
 
-FReply SDataprepActionBlock::OnMouseButtonDown(const FGeometry& MyGeometry, const FPointerEvent& MouseEvent)
-{
-#ifndef NO_BLUEPRINT
-	if ( !bIsSimplifiedGraph && MouseEvent.GetEffectingButton() ==  EKeys::LeftMouseButton )
-	{
-		return FReply::Handled().DetectDrag( AsShared(), EKeys::LeftMouseButton );
-	}
-#else
-	if ( MouseEvent.GetEffectingButton() ==  EKeys::LeftMouseButton )
-	{
-		return FReply::Handled().DetectDrag( AsShared(), EKeys::LeftMouseButton );
-	}
-#endif
-
-	return FReply::Unhandled();
-}
-
-FReply SDataprepActionBlock::OnDragDetected(const FGeometry& MyGeometry, const FPointerEvent& MouseEvent)
-{
-	check( DataprepActionContext );
-
-	// Is the step moved or copied
-	TSharedRef<bool> IsMoved = MakeShared<bool>( true );
-
-	// Callback that move or copy the operation
-	FDataprepGraphOperation Operation = FDataprepGraphOperation::CreateLambda( [InitialContext = *DataprepActionContext.Get(), IsMoved] (const FDataprepSchemaActionContext& Context)
-		{
-
-			UDataprepActionAsset* InitialAction = InitialContext.DataprepActionPtr.Get();
-			UDataprepActionAsset* DroppedOnAction = Context.DataprepActionPtr.Get();
-			if ( InitialAction && DroppedOnAction )
-			{
-				if ( InitialAction == DroppedOnAction && IsMoved.Get() )
-				{
-					// Move the position of the step in the action
-					if ( InitialContext.StepIndex != INDEX_NONE && InitialContext.StepIndex != Context.StepIndex )
-					{
-						if ( Context.StepIndex == INDEX_NONE )
-						{
-							if ( InitialContext.StepIndex != DroppedOnAction->GetStepsCount() - 1 )
-							{
-								// The step will be move at the end of the array
-								DroppedOnAction->MoveStep( InitialContext.StepIndex, DroppedOnAction->GetStepsCount() - 1 );
-							}
-							else
-							{
-								return false;
-							}
-						}
-						else
-						{
-							DroppedOnAction->MoveStep( InitialContext.StepIndex, Context.StepIndex );
-						}
-						return true;
-					}
-				}
-				else
-				{
-					// Move the step to a another action or duplicate the step
-					UDataprepActionStep* ActionStep = InitialContext.DataprepActionStepPtr.Get();
-					if ( ActionStep )
-					{
-						uint32 NewStepIndex = DroppedOnAction->AddStep( ActionStep );
-						if ( Context.StepIndex != INDEX_NONE )
-						{
-							DroppedOnAction->MoveStep( NewStepIndex, Context.StepIndex );
-						}
-						
-						if ( IsMoved.Get() )
-						{
-							InitialAction->RemoveStep( InitialContext.StepIndex );
-						}
-						return true;
-					}
-				}
-			}
-
-			return false;
-		});
-
-	TSharedRef<FDataprepDragDropOp> DragDropOperation = FDataprepDragDropOp::New( MoveTemp( Operation ) );
-
-	// The pre drop pop the contextual menu allowing the user to select between copy, move and simply cancel the drag and drop.
-	FDataprepPreDropConfirmation PreDropConfirmation = FDataprepPreDropConfirmation::CreateLambda( [ SourceWidget = AsShared(), IsMoved ] (const FDataprepSchemaActionContext& Context, TFunction<void ()> ConfirmationCallback)
-		{
-			auto CopyDraggedItems = [ConfirmationCallback, IsMoved] ()
-				{
-					IsMoved.Get() = false;
-					ConfirmationCallback();
-				};
-
-			auto MoveDraggedItems = [ConfirmationCallback, IsMoved] ()
-				{
-					IsMoved.Get() = true;
-					ConfirmationCallback();
-				};
-
-			FModifierKeysState ModifierKeyState = FSlateApplication::Get().GetModifierKeys();
-
-			bool bCopyKeyDown = ModifierKeyState.IsControlDown()  || ModifierKeyState.IsCommandDown();
-			bool bMoveKeyDown = ModifierKeyState.IsAltDown();
-
-			if ( bCopyKeyDown && !bMoveKeyDown )
-			{
-				// The user is using the shortcut to copy the dragged items
-				IsMoved.Get() = false;
-				ConfirmationCallback();
-			}
-			else if ( bMoveKeyDown && !bCopyKeyDown )
-			{
-				// The user using the shortcut to move the dragged items
-				IsMoved.Get() = true;
-				ConfirmationCallback();
-			}
-			else
-			{
-				FMenuBuilder MenuBuilder(/*bInShouldCloseWindowAfterMenuSelection=*/true, nullptr);
-				const FText MoveCopyHeaderString = LOCTEXT("AssetViewDropMenuHeading", "Move or Copy");
-				MenuBuilder.BeginSection("PathAssetMoveCopy", MoveCopyHeaderString);
-				{
-					MenuBuilder.AddMenuEntry(
-						LOCTEXT("DragDropMove", "Move Here"),
-						LOCTEXT("DragDropMoveTooltip", "Move the dragged items."),
-						FSlateIcon(),
-						FUIAction( FExecuteAction::CreateLambda( MoveDraggedItems ) )
-					);
-					MenuBuilder.AddMenuEntry(
-						LOCTEXT("DragDropCopy", "Copy Here"),
-						LOCTEXT("DragDropCopyTooltip", "Copy the dragged items."),
-						FSlateIcon(),
-						FUIAction( FExecuteAction::CreateLambda( CopyDraggedItems ) )
-					);
-				}
-				MenuBuilder.EndSection();
-
-				TSharedPtr<IMenu> Menu = FSlateApplication::Get().PushMenu(
-					FSlateApplication::Get().GetActiveTopLevelWindow().ToSharedRef(),
-					FWidgetPath(),
-					MenuBuilder.MakeWidget(),
-					FSlateApplication::Get().GetCursorPos(),
-					FPopupTransitionEffect( FPopupTransitionEffect::ContextMenu )
-					);
-
-				if ( GEditor && GEditor->Trans && Menu )
-				{
-					UTransBuffer* TransBuffer = CastChecked<UTransBuffer>( GEditor->Trans );
-					FDelegateHandle Handle = TransBuffer->OnBeforeRedoUndo().AddLambda( [Menu] (const struct FTransactionContext&)
-						{
-							FSlateApplication::Get().DismissMenu( Menu );
-						});
-
-					Menu->GetOnMenuDismissed().AddLambda( [Handle] (TSharedRef<IMenu>)
-						{
-							if ( GEditor && GEditor->Trans )
-							{
-								UTransBuffer* TransBuffer = CastChecked<UTransBuffer>( GEditor->Trans );
-								TransBuffer->OnBeforeRedoUndo().Remove( Handle );
-							}
-						} );
-				}
-			}
-		} );
-
-	DragDropOperation->SetPreDropConfirmation( MoveTemp( PreDropConfirmation ) );
-
-	return FReply::Handled().BeginDragDrop( MoveTemp( DragDropOperation ) );
-}
-
 FSlateColor SDataprepActionBlock::GetOutlineColor() const
 {
 	return FDataprepEditorStyle::GetColor( "DataprepAction.OutlineColor" );
diff --git a/Engine/Plugins/Enterprise/DataprepEditor/Source/DataprepEditor/Private/Widgets/DataprepGraph/SDataprepActionBlock.h b/Engine/Plugins/Enterprise/DataprepEditor/Source/DataprepEditor/Private/Widgets/DataprepGraph/SDataprepActionBlock.h
--- a/Engine/Plugins/Enterprise/DataprepEditor/Source/DataprepEditor/Private/Widgets/DataprepGraph/SDataprepActionBlock.h
+++ b/Engine/Plugins/Enterprise/DataprepEditor/Source/DataprepEditor/Private/Widgets/DataprepGraph/SDataprepActionBlock.h
@@ -14,15 +14,7 @@ struct FTextBlockStyle;
  */
 class SDataprepActionBlock : public SCompoundWidget
 {
-#ifndef NO_BLUEPRINT
-	SLATE_BEGIN_ARGS(SDataprepActionBlock)
-	: _IsSimplified(false)
-	{}
-
-		SLATE_ARGUMENT( bool, IsSimplified )
-#else
 	SLATE_BEGIN_ARGS(SDataprepActionBlock) {}
-#endif
 	SLATE_END_ARGS()
 
 	void Construct(const FArguments& InArgs, const TSharedRef<FDataprepSchemaActionContext>& DataprepActionContext);
@@ -33,8 +25,6 @@ class SDataprepActionBlock : public SCompoundWidget
 		const FVector2D ChildSize = ChildSlot.GetWidget()->GetDesiredSize();
 		return FVector2D( FMath::Max(280.f, ChildSize.X) , ChildSize.Y );
 	}
-	virtual FReply OnMouseButtonDown(const FGeometry& MyGeometry, const FPointerEvent& MouseEvent) override;
-	virtual FReply OnDragDetected(const FGeometry& MyGeometry, const FPointerEvent& MouseEvent) override;
 	// End of SWidget Interface
 
 	/**
@@ -52,15 +42,7 @@ public:
 	 */
 	virtual TSharedRef<SWidget> GetTitleWidget();
 
-#ifndef NO_BLUEPRINT
-	// Temporary boolean to indicate if the widget has been instantiated in a regular BP graph or the simplified one
-	bool bIsSimplifiedGraph;
-#endif
-
 protected:
-
-	void ConstructForSimplified();
-
 	/**
 	 * Return the block title.
 	 * This is used in the implementation of SDataprepActionBlock::GetTitleWidget
diff --git a/Engine/Plugins/Enterprise/DataprepEditor/Source/DataprepEditor/Private/Widgets/DataprepGraph/SDataprepFilter.cpp b/Engine/Plugins/Enterprise/DataprepEditor/Source/DataprepEditor/Private/Widgets/DataprepGraph/SDataprepFilter.cpp
--- a/Engine/Plugins/Enterprise/DataprepEditor/Source/DataprepEditor/Private/Widgets/DataprepGraph/SDataprepFilter.cpp
+++ b/Engine/Plugins/Enterprise/DataprepEditor/Source/DataprepEditor/Private/Widgets/DataprepGraph/SDataprepFilter.cpp
@@ -34,7 +34,7 @@ void SDataprepFilter::Construct(const FArguments& InArgs, UDataprepFilter& InFil
 
 	bIsPreviewed = InArgs._IsPreviewed;
 
-	SDataprepActionBlock::Construct( SDataprepActionBlock::FArguments().IsSimplified(InArgs._IsSimplified), InDataprepActionContext );
+	SDataprepActionBlock::Construct( SDataprepActionBlock::FArguments(), InDataprepActionContext );
 }
 
 void SDataprepFilter::Tick(const FGeometry& AllottedGeometry, const double InCurrentTime, const float InDeltaTime)
diff --git a/Engine/Plugins/Enterprise/DataprepEditor/Source/DataprepEditor/Private/Widgets/DataprepGraph/SDataprepFilter.h b/Engine/Plugins/Enterprise/DataprepEditor/Source/DataprepEditor/Private/Widgets/DataprepGraph/SDataprepFilter.h
--- a/Engine/Plugins/Enterprise/DataprepEditor/Source/DataprepEditor/Private/Widgets/DataprepGraph/SDataprepFilter.h
+++ b/Engine/Plugins/Enterprise/DataprepEditor/Source/DataprepEditor/Private/Widgets/DataprepGraph/SDataprepFilter.h
@@ -18,11 +18,8 @@ class SDataprepFilter : public SDataprepActionBlock, public FGCObject
 {
 public:
 	SLATE_BEGIN_ARGS(SDataprepFilter)
-
-		: _IsSimplified(false)
-		, _IsPreviewed(false)
+		: _IsPreviewed(false)
 	{}
-		SLATE_ARGUMENT( bool, IsSimplified )
 		SLATE_ARGUMENT( bool, IsPreviewed )
 	SLATE_END_ARGS()
 
diff --git a/Engine/Plugins/Enterprise/DataprepEditor/Source/DataprepEditor/Private/Widgets/DataprepGraph/SDataprepGraphActionStepNode.cpp b/Engine/Plugins/Enterprise/DataprepEditor/Source/DataprepEditor/Private/Widgets/DataprepGraph/SDataprepGraphActionStepNode.cpp
--- a/Engine/Plugins/Enterprise/DataprepEditor/Source/DataprepEditor/Private/Widgets/DataprepGraph/SDataprepGraphActionStepNode.cpp
+++ b/Engine/Plugins/Enterprise/DataprepEditor/Source/DataprepEditor/Private/Widgets/DataprepGraph/SDataprepGraphActionStepNode.cpp
@@ -127,7 +127,7 @@ void SDataprepGraphActionStepNode::UpdateGraphNode()
 			if (StepType == UDataprepOperation::StaticClass())
 			{
 				UDataprepOperation* Operation = static_cast<UDataprepOperation*>( StepObject );
-				ActionStepBlockPtr = StaticCastSharedRef<SDataprepActionBlock>( SNew(SDataprepOperation, Operation, StepData).IsSimplified(true) );
+				ActionStepBlockPtr = StaticCastSharedRef<SDataprepActionBlock>( SNew(SDataprepOperation, Operation, StepData) );
 			}
 			else if (StepType == UDataprepFilter::StaticClass())
 			{
@@ -136,20 +136,16 @@ void SDataprepGraphActionStepNode::UpdateGraphNode()
 				ActionStepBlockPtr = StaticCastSharedRef<SDataprepActionBlock>( 
 					SNew( SDataprepFilter, *Filter, StepData )
 						.IsPreviewed( bIsPreviewed )
-						.IsSimplified(true)
 					);
 			}
 			else if (StepType == UDataprepSelectionTransform::StaticClass())
 			{
 				UDataprepSelectionTransform* SelectionTransform = static_cast<UDataprepSelectionTransform*>( StepObject );
-				ActionStepBlockPtr = StaticCastSharedRef<SDataprepActionBlock>( SNew(SDataprepSelectionTransform, SelectionTransform, StepData).IsSimplified(true) );
+				ActionStepBlockPtr = StaticCastSharedRef<SDataprepActionBlock>( SNew(SDataprepSelectionTransform, SelectionTransform, StepData) );
 			}
 
 			if(ActionStepBlockPtr.IsValid())
 			{
-#ifndef NO_BLUEPRINT
-				ActionStepBlockPtr->bIsSimplifiedGraph = true;
-#endif
 				ActionBlockPtr = ActionStepBlockPtr->AsShared();
 			}
 		}
diff --git a/Engine/Plugins/Enterprise/DataprepEditor/Source/DataprepEditor/Private/Widgets/DataprepGraph/SDataprepGraphEditor.cpp b/Engine/Plugins/Enterprise/DataprepEditor/Source/DataprepEditor/Private/Widgets/DataprepGraph/SDataprepGraphEditor.cpp
--- a/Engine/Plugins/Enterprise/DataprepEditor/Source/DataprepEditor/Private/Widgets/DataprepGraph/SDataprepGraphEditor.cpp
+++ b/Engine/Plugins/Enterprise/DataprepEditor/Source/DataprepEditor/Private/Widgets/DataprepGraph/SDataprepGraphEditor.cpp
@@ -120,17 +120,9 @@ void SDataprepGraphEditor::Construct(const FArguments& InArgs, UDataprepAsset* I
 
 	DataprepAssetPtr->GetOnActionChanged().AddSP(this, &SDataprepGraphEditor::OnDataprepAssetActionChanged);
 
-	// #ueent_toremove: Temp code for the nodes development
-	if(UBlueprint* RecipeBP = InDataprepAsset->GetRecipeBP())
-	{
-		RecipeBP->OnChanged().AddSP(this, &SDataprepGraphEditor::OnPipelineChanged);
-	}
-	// end of temp code for nodes development
-
 	SetCanTick(true);
 
 	bIsComplete = false;
-	bMustRearrange = false;
 
 	LastLocalSize = FVector2D::ZeroVector;
 	LastZoomAmount = BIG_NUMBER;
@@ -154,15 +146,6 @@ void SDataprepGraphEditor::NotifyGraphChanged()
 	SGraphEditor::NotifyGraphChanged();
 }
 
-// #ueent_toremove: Temp code for the nodes development
-void SDataprepGraphEditor::OnPipelineChanged(UBlueprint* InBlueprint)
-{
-	if (InBlueprint)
-	{
-		NotifyGraphChanged();
-	}
-}
-
 void SDataprepGraphEditor::OnDataprepAssetActionChanged(UObject* InObject, FDataprepAssetChangeType ChangeType)
 {
 	switch(ChangeType)
@@ -226,7 +209,6 @@ void SDataprepGraphEditor::CacheDesiredSize(float InLayoutScaleMultiplier)
 		if(TrackGraphNodePtr.IsValid())
 		{
 			bIsComplete = TrackGraphNodePtr.Pin()->RefreshLayout();
-			bMustRearrange = true;
 		}
 	}
 }
@@ -262,17 +244,10 @@ void SDataprepGraphEditor::UpdateLayout( const FVector2D& LocalSize, const FVect
 {
 	if(SDataprepGraphTrackNode* TrackGraphNode = TrackGraphNodePtr.Pin().Get())
 	{
-		if(LastZoomAmount != ZoomAmount)
-		{
-			WorkingArea = TrackGraphNode->Update();
-		}
+		WorkingArea = TrackGraphNode->Update();
 
 		if( !LocalSize.Equals(LastLocalSize) )
 		{
-			bMustRearrange = true;
-
-			WorkingArea = TrackGraphNode->Update();
-
 			LastLocalSize = LocalSize;
 
 			// Force a re-compute of the view location
diff --git a/Engine/Plugins/Enterprise/DataprepEditor/Source/DataprepEditor/Private/Widgets/DataprepGraph/SDataprepGraphEditor.h b/Engine/Plugins/Enterprise/DataprepEditor/Source/DataprepEditor/Private/Widgets/DataprepGraph/SDataprepGraphEditor.h
--- a/Engine/Plugins/Enterprise/DataprepEditor/Source/DataprepEditor/Private/Widgets/DataprepGraph/SDataprepGraphEditor.h
+++ b/Engine/Plugins/Enterprise/DataprepEditor/Source/DataprepEditor/Private/Widgets/DataprepGraph/SDataprepGraphEditor.h
@@ -88,9 +88,6 @@ public:
 	virtual void NotifyGraphChanged() override;
 	// End of SGraphEditor overrides
 
-	// #ueent_toremove: Temp code for the nodes development
-	void OnPipelineChanged(UBlueprint* InBlueprint);
-
 	/** Called when a change has occurred in the set of the Dataprep asset's actions */
 	void OnDataprepAssetActionChanged(UObject* InObject, FDataprepAssetChangeType ChangeType);
 
@@ -142,9 +139,6 @@ private:
 	/** When false, indicates the graph editor has not been drawn yet */
 	mutable bool bIsComplete;
 
-	/** Indicates layout must be recomputed */
-	mutable bool bMustRearrange;
-
 	/** Last size of the window displaying the graph's canvas */
 	FVector2D LastLocalSize;
 
diff --git a/Engine/Plugins/Enterprise/DataprepEditor/Source/DataprepEditor/Private/Widgets/DataprepGraph/SDataprepGraphTrackNode.cpp b/Engine/Plugins/Enterprise/DataprepEditor/Source/DataprepEditor/Private/Widgets/DataprepGraph/SDataprepGraphTrackNode.cpp
--- a/Engine/Plugins/Enterprise/DataprepEditor/Source/DataprepEditor/Private/Widgets/DataprepGraph/SDataprepGraphTrackNode.cpp
+++ b/Engine/Plugins/Enterprise/DataprepEditor/Source/DataprepEditor/Private/Widgets/DataprepGraph/SDataprepGraphTrackNode.cpp
@@ -533,34 +533,37 @@ FSlateRect SDataprepGraphTrackNode::Update()
 
 	if(SDataprepGraphTrackWidget* TrackWidget = TrackWidgetPtr.Get())
 	{
-		RefreshLayout();
+		const FVector2D& WorkingSize = TrackWidget->WorkingSize;
+		const FVector2D TrackPosition = GetPosition().GetAbs();
 
-		SGraphPanel& GraphPanel = *OwnerGraphPanelPtr.Pin();
+		if(TrackWidget->DragIndicatorIndex == INDEX_NONE && !bNodeDragging)
+		{
+			RefreshLayout();
 
-		const FVector2D& WorkingSize = TrackWidget->WorkingSize;
+			SGraphPanel& GraphPanel = *OwnerGraphPanelPtr.Pin();
 
-		// Determine canvas offset attribute in track widget's coordinates
-		const FVector2D PanelSize = GraphPanel.GetTickSpaceGeometry().GetLocalSize() / GraphPanel.GetZoomAmount();
+			// Determine canvas offset attribute in track widget's coordinates
+			const FVector2D PanelSize = GraphPanel.GetTickSpaceGeometry().GetLocalSize() / GraphPanel.GetZoomAmount();
 
-		const FVector2D TrackPosition = GetPosition().GetAbs();
-		const FVector2D TargetSize = FVector2D(FMath::Max(WorkingSize.X, PanelSize.X), FMath::Max(WorkingSize.Y, PanelSize.Y)) + TrackPosition + 20.f;
+			const FVector2D TargetSize = FVector2D(FMath::Max(WorkingSize.X, PanelSize.X), FMath::Max(WorkingSize.Y, PanelSize.Y)) + TrackPosition + 20.f;
 
-		FMargin& CanvasOffset = TrackWidgetPtr->CanvasOffset;
+			FMargin& CanvasOffset = TrackWidgetPtr->CanvasOffset;
 
-		CanvasOffset.Left = -TrackPosition.X - 10.f;
-		CanvasOffset.Top = -TrackPosition.Y - 10.f;
+			CanvasOffset.Left = -TrackPosition.X - 10.f;
+			CanvasOffset.Top = -TrackPosition.Y - 10.f;
 
-		if(TargetSize.X > CanvasOffset.Right)
-		{
-			CanvasOffset.Right = TargetSize.X;
-		}
+			if(TargetSize.X > CanvasOffset.Right)
+			{
+				CanvasOffset.Right = TargetSize.X;
+			}
 
-		if(TargetSize.Y > CanvasOffset.Bottom)
-		{
-			CanvasOffset.Bottom = TargetSize.Y;
-		}
+			if(TargetSize.Y > CanvasOffset.Bottom)
+			{
+				CanvasOffset.Bottom = TargetSize.Y;
+			}
 
-		TrackWidgetPtr->CanvasSlot->OffsetAttr.Set(CanvasOffset);
+			TrackWidgetPtr->CanvasSlot->OffsetAttr.Set(CanvasOffset);
+		}
 
 		WorkingArea = FSlateRect(FVector2D::ZeroVector, WorkingSize + TrackPosition);
 	}
diff --git a/Engine/Plugins/Enterprise/DataprepEditor/Source/DataprepEditor/Private/Widgets/DataprepGraph/SDataprepOperation.cpp b/Engine/Plugins/Enterprise/DataprepEditor/Source/DataprepEditor/Private/Widgets/DataprepGraph/SDataprepOperation.cpp
--- a/Engine/Plugins/Enterprise/DataprepEditor/Source/DataprepEditor/Private/Widgets/DataprepGraph/SDataprepOperation.cpp
+++ b/Engine/Plugins/Enterprise/DataprepEditor/Source/DataprepEditor/Private/Widgets/DataprepGraph/SDataprepOperation.cpp
@@ -18,7 +18,7 @@ void SDataprepOperation::Construct(const FArguments& InArgs, UDataprepOperation*
 	TAttribute<FText> TooltipTextAttribute = MakeAttributeSP( this, &SDataprepOperation::GetTooltipText );
 	SetToolTipText( TooltipTextAttribute );
 
-	SDataprepActionBlock::Construct( SDataprepActionBlock::FArguments().IsSimplified(InArgs._IsSimplified), InDataprepActionContext );
+	SDataprepActionBlock::Construct( SDataprepActionBlock::FArguments(), InDataprepActionContext );
 }
 
 FSlateColor SDataprepOperation::GetOutlineColor() const
diff --git a/Engine/Plugins/Enterprise/DataprepEditor/Source/DataprepEditor/Private/Widgets/DataprepGraph/SDataprepOperation.h b/Engine/Plugins/Enterprise/DataprepEditor/Source/DataprepEditor/Private/Widgets/DataprepGraph/SDataprepOperation.h
--- a/Engine/Plugins/Enterprise/DataprepEditor/Source/DataprepEditor/Private/Widgets/DataprepGraph/SDataprepOperation.h
+++ b/Engine/Plugins/Enterprise/DataprepEditor/Source/DataprepEditor/Private/Widgets/DataprepGraph/SDataprepOperation.h
@@ -14,14 +14,7 @@ struct FDataprepSchemaActionContext;
 
 class SDataprepOperation : public SDataprepActionBlock, public FGCObject
 {
-#ifndef NO_BLUEPRINT
-	SLATE_BEGIN_ARGS(SDataprepOperation)
-	: _IsSimplified(false)
-	{}
-		SLATE_ARGUMENT( bool, IsSimplified )
-#else
 	SLATE_BEGIN_ARGS(SDataprepOperation) {}
-#endif
 	SLATE_END_ARGS()
 
 	void Construct(const FArguments& InArgs, UDataprepOperation* InOperation, const TSharedRef<FDataprepSchemaActionContext>& InDataprepActionContext);
diff --git a/Engine/Plugins/Enterprise/DataprepEditor/Source/DataprepEditor/Private/Widgets/DataprepGraph/SDataprepSelectionTransform.cpp b/Engine/Plugins/Enterprise/DataprepEditor/Source/DataprepEditor/Private/Widgets/DataprepGraph/SDataprepSelectionTransform.cpp
--- a/Engine/Plugins/Enterprise/DataprepEditor/Source/DataprepEditor/Private/Widgets/DataprepGraph/SDataprepSelectionTransform.cpp
+++ b/Engine/Plugins/Enterprise/DataprepEditor/Source/DataprepEditor/Private/Widgets/DataprepGraph/SDataprepSelectionTransform.cpp
@@ -18,7 +18,7 @@ void SDataprepSelectionTransform::Construct(const FArguments& InArgs, UDataprepS
 	TAttribute<FText> TooltipTextAttribute = MakeAttributeSP(this, &SDataprepSelectionTransform::GetTooltipText);
 	SetToolTipText(TooltipTextAttribute);
 
-	SDataprepActionBlock::Construct(SDataprepActionBlock::FArguments().IsSimplified(InArgs._IsSimplified), InDataprepActionContext);
+	SDataprepActionBlock::Construct(SDataprepActionBlock::FArguments(), InDataprepActionContext);
 }
 
 FSlateColor SDataprepSelectionTransform::GetOutlineColor() const
diff --git a/Engine/Plugins/Enterprise/DataprepEditor/Source/DataprepEditor/Private/Widgets/DataprepGraph/SDataprepSelectionTransform.h b/Engine/Plugins/Enterprise/DataprepEditor/Source/DataprepEditor/Private/Widgets/DataprepGraph/SDataprepSelectionTransform.h
--- a/Engine/Plugins/Enterprise/DataprepEditor/Source/DataprepEditor/Private/Widgets/DataprepGraph/SDataprepSelectionTransform.h
+++ b/Engine/Plugins/Enterprise/DataprepEditor/Source/DataprepEditor/Private/Widgets/DataprepGraph/SDataprepSelectionTransform.h
@@ -14,10 +14,7 @@ struct FDataprepSchemaActionContext;
 
 class SDataprepSelectionTransform : public SDataprepActionBlock, public FGCObject
 {
-	SLATE_BEGIN_ARGS(SDataprepSelectionTransform)
-	: _IsSimplified(false)
-	{}
-		SLATE_ARGUMENT( bool, IsSimplified )
+	SLATE_BEGIN_ARGS(SDataprepSelectionTransform){}
 	SLATE_END_ARGS();
 
 	void Construct(const FArguments& InArgs, UDataprepSelectionTransform* InTransform, const TSharedRef<FDataprepSchemaActionContext>& InDataprepActionContext);
diff --git a/Engine/Plugins/Enterprise/DataprepEditor/Source/DataprepEditor/Private/Widgets/DataprepGraph/SGraphNodeK2DataprepAction.h b/Engine/Plugins/Enterprise/DataprepEditor/Source/DataprepEditor/Private/Widgets/DataprepGraph/SGraphNodeK2DataprepAction.h
deleted file mode 100644
diff --git a/Engine/Plugins/Enterprise/DataprepEditor/Source/DataprepEditor/Private/Widgets/DataprepGraph/SGraphNodek2DataprepAction.cpp b/Engine/Plugins/Enterprise/DataprepEditor/Source/DataprepEditor/Private/Widgets/DataprepGraph/SGraphNodek2DataprepAction.cpp
deleted file mode 100644
diff --git a/Engine/Plugins/Enterprise/DataprepEditor/Source/DataprepEditor/Private/Widgets/SDataprepActionMenu.cpp b/Engine/Plugins/Enterprise/DataprepEditor/Source/DataprepEditor/Private/Widgets/SDataprepActionMenu.cpp
--- a/Engine/Plugins/Enterprise/DataprepEditor/Source/DataprepEditor/Private/Widgets/SDataprepActionMenu.cpp
+++ b/Engine/Plugins/Enterprise/DataprepEditor/Source/DataprepEditor/Private/Widgets/SDataprepActionMenu.cpp
@@ -71,40 +71,6 @@ void SDataprepActionMenu::CollectActions(FGraphActionListBuilderBase& OutActions
 
 void SDataprepActionMenu::OnActionSelected(const TArray< TSharedPtr<FEdGraphSchemaAction> >& SelectedActions, ESelectInfo::Type InSelectionType)
 {
-#ifndef NO_BLUEPRINT
-	/**
-	 * Create a node if needed and execute the action
-	 * return true if the action was properly executed
-	 */
-	auto ExecuteActionBP = [this](FDataprepSchemaAction& DataprepAction) -> bool
-	{
-		UK2Node_DataprepAction* ActionNode = nullptr;
-
-		if ( ShouldCreateNewNode() )
-		{
-			ActionNode = DataprepSchemaActionUtils::SpawnEdGraphNode<UK2Node_DataprepAction>( *GraphObj, NewNodePosition );
-			ActionNode->CreateDataprepActionAsset();
-			Context.DataprepActionPtr = MakeWeakObjectPtr<UDataprepActionAsset>( ActionNode->GetDataprepAction() );
-
-			if ( DraggedFromPins.Num() > 0 )
-			{
-				ActionNode->AutowireNewNode( DraggedFromPins[0] );
-			}
-		}
-
-		if ( !DataprepAction.ExecuteAction( Context ) )
-		{
-			if ( ActionNode )
-			{
-				// Remove the created node
-				GraphObj->RemoveNode( ActionNode );
-			}
-			return false;
-		}
-
-		return true;
-	};
-#endif
 	auto ExecuteAction = [this](FDataprepSchemaAction& DataprepAction, UDataprepGraph* DataprepGraph) -> bool
 	{
 		UDataprepAsset* DataprepAsset = DataprepGraph ? DataprepGraph->GetDataprepAsset() : nullptr;
@@ -134,25 +100,10 @@ void SDataprepActionMenu::OnActionSelected(const TArray< TSharedPtr<FEdGraphSche
 						Transaction.Cancel();
 					}
 				}
-#ifndef NO_BLUEPRINT
-				else if ( !ExecuteActionBP( *DataprepAction ) )
-				{
-					Transaction.Cancel();
-				}
-#endif
 			}
-			else
+			else if(UDataprepGraph* DataprepGraph = Cast<UDataprepGraph>(GraphObj))
 			{
-				if(UDataprepGraph* DataprepGraph = Cast<UDataprepGraph>(GraphObj))
-				{
-					ExecuteAction( *DataprepAction, DataprepGraph );
-				}
-#ifndef NO_BLUEPRINT
-				else
-				{
-					ExecuteActionBP( *DataprepAction );
-				}
-#endif
+				ExecuteAction( *DataprepAction, DataprepGraph );
 			}
 		}
 	}
diff --git a/Engine/Plugins/Enterprise/DataprepEditor/Source/DataprepEditor/Private/Widgets/SDataprepProducersWidget.cpp b/Engine/Plugins/Enterprise/DataprepEditor/Source/DataprepEditor/Private/Widgets/SDataprepProducersWidget.cpp
--- a/Engine/Plugins/Enterprise/DataprepEditor/Source/DataprepEditor/Private/Widgets/SDataprepProducersWidget.cpp
+++ b/Engine/Plugins/Enterprise/DataprepEditor/Source/DataprepEditor/Private/Widgets/SDataprepProducersWidget.cpp
@@ -253,10 +253,8 @@ SDataprepProducersWidget::~SDataprepProducersWidget()
 	}
 }
 
-void SDataprepProducersWidget::Construct(const FArguments & InArgs, UDataprepAssetProducers* InAssetProducersPtr, TSharedPtr<FUICommandList>& InCommandList)
+void SDataprepProducersWidget::Construct(const FArguments & InArgs, UDataprepAssetProducers* InAssetProducersPtr)
 {
-	CommandList = InCommandList;
-
 	AssetProducersPtr = InAssetProducersPtr;
 	AssetProducersPtr->GetOnChanged().AddSP( this, &SDataprepProducersWidget::OnDataprepProducersChanged );
 
@@ -324,6 +322,8 @@ void SDataprepProducersWidget::OnDataprepProducersChanged(FDataprepAssetChangeTy
 TSharedRef<SWidget> SDataprepProducersWidget::CreateAddProducerMenuWidget()
 {
 	const bool bShouldCloseWindowAfterMenuSelection = true;
+
+	TSharedPtr<FUICommandList> CommandList = MakeShared<FUICommandList>();
 	FMenuBuilder MenuBuilder(bShouldCloseWindowAfterMenuSelection, CommandList);
 
 	MenuBuilder.BeginSection("AddNewProducer", LOCTEXT("DataprepProducersWidget_AddImports", "Add Producer"));
@@ -375,9 +375,9 @@ void SDataprepProducersWidget::OnAddProducer( UClass* ProducerClass )
 	TreeView->Refresh();
 }
 
-TSharedRef<SWidget> FDataprepAssetProducersDetails::CreateWidget(UDataprepAssetProducers* Producers, TSharedPtr<FUICommandList>& CommandList)
+TSharedRef<SWidget> FDataprepAssetProducersDetails::CreateWidget(UDataprepAssetProducers* Producers)
 {
-	return SNew(SDataprepProducersWidget, Producers, CommandList );
+	return SNew(SDataprepProducersWidget, Producers );
 }
 
 void FDataprepAssetProducersDetails::CustomizeDetails(IDetailLayoutBuilder& DetailBuilder)
@@ -395,8 +395,6 @@ void FDataprepAssetProducersDetails::CustomizeDetails(IDetailLayoutBuilder& Deta
 	FName CategoryName = CategoryNames.Num() > 0 ? CategoryNames[0] : FName( TEXT("DatasmithDirProducer") );
 	IDetailCategoryBuilder& CategoryBuilder = DetailBuilder.EditCategory( CategoryName, FText::GetEmpty(), ECategoryPriority::Important );
 
-	TSharedPtr<FUICommandList> CommandList = DetailBuilder.GetDetailsView()->GetHostCommandList();
-
 	FDetailWidgetRow& CustomRow = CategoryBuilder.AddCustomRow( FText::GetEmpty() )
 	.NameContent()
 	[
@@ -404,7 +402,7 @@ void FDataprepAssetProducersDetails::CustomizeDetails(IDetailLayoutBuilder& Deta
 	]
 	.ValueContent()
 	[
-		CreateWidget( Producers, CommandList )
+		CreateWidget( Producers )
 	];
 }
 
diff --git a/Engine/Plugins/Enterprise/DataprepEditor/Source/DataprepEditor/Private/Widgets/SDataprepProducersWidget.h b/Engine/Plugins/Enterprise/DataprepEditor/Source/DataprepEditor/Private/Widgets/SDataprepProducersWidget.h
--- a/Engine/Plugins/Enterprise/DataprepEditor/Source/DataprepEditor/Private/Widgets/SDataprepProducersWidget.h
+++ b/Engine/Plugins/Enterprise/DataprepEditor/Source/DataprepEditor/Private/Widgets/SDataprepProducersWidget.h
@@ -111,7 +111,7 @@ public:
 
 	~SDataprepProducersWidget();
 
-	void Construct(const FArguments& InArgs, UDataprepAssetProducers* InAssetProducersPtr, TSharedPtr<FUICommandList>& CommandList);
+	void Construct(const FArguments& InArgs, UDataprepAssetProducers* InAssetProducersPtr);
 
 	TSharedRef<SWidget> CreateAddProducerMenuWidget();
 
@@ -127,7 +127,6 @@ private:
 
 private:
 	TSharedPtr<SWidget> AddNewMenu;
-	TSharedPtr<FUICommandList> CommandList;
 	TWeakObjectPtr<UDataprepAssetProducers> AssetProducersPtr;
 	TSharedPtr<SDataprepProducersTreeView> TreeView;
 };
@@ -142,5 +141,5 @@ public:
 	virtual void CustomizeDetails( IDetailLayoutBuilder& DetailBuilder ) override;
 
 private:
-	TSharedRef<SWidget> CreateWidget(UDataprepAssetProducers* Producers, TSharedPtr<FUICommandList>& CommandList);
+	TSharedRef<SWidget> CreateWidget(UDataprepAssetProducers* Producers);
 };
diff --git a/Engine/Plugins/Enterprise/DataprepEditor/Source/DataprepEditorScriptingUtilities/Private/EditorDataprepAssetLibrary.cpp b/Engine/Plugins/Enterprise/DataprepEditor/Source/DataprepEditorScriptingUtilities/Private/EditorDataprepAssetLibrary.cpp
--- a/Engine/Plugins/Enterprise/DataprepEditor/Source/DataprepEditorScriptingUtilities/Private/EditorDataprepAssetLibrary.cpp
+++ b/Engine/Plugins/Enterprise/DataprepEditor/Source/DataprepEditorScriptingUtilities/Private/EditorDataprepAssetLibrary.cpp
@@ -194,7 +194,7 @@ void UEditorDataprepAssetLibrary::RemoveAction(UDataprepAsset* DataprepAsset, in
 {
 	if ( DataprepAsset )
 	{
-		DataprepAsset->RemoveActionUsingBP( Index );
+		DataprepAsset->RemoveAction( Index );
 	}
 	else
 	{
@@ -244,7 +244,7 @@ void UEditorDataprepAssetLibrary::SwapActions(UDataprepAsset* DataprepAsset, int
 {
 	if ( DataprepAsset )
 	{
-		DataprepAsset->SwapActionsUsingBP( FirstActionIndex, SecondActionIndex );
+		DataprepAsset->SwapActions( FirstActionIndex, SecondActionIndex );
 	}
 	else
 	{
