Fixed issue with recipe Graph pane not in the default layout

#rnx

diff --git a/Engine/Plugins/Enterprise/DataprepEditor/Source/DataprepEditor/Private/DataprepEditor.cpp b/Engine/Plugins/Enterprise/DataprepEditor/Source/DataprepEditor/Private/DataprepEditor.cpp
--- a/Engine/Plugins/Enterprise/DataprepEditor/Source/DataprepEditor/Private/DataprepEditor.cpp
+++ b/Engine/Plugins/Enterprise/DataprepEditor/Source/DataprepEditor/Private/DataprepEditor.cpp
@@ -937,7 +937,7 @@ TSharedRef<SDockTab> FDataprepEditor::SpawnTabGraphEditor(const FSpawnTabArgs &
 
 TSharedRef<FTabManager::FLayout> FDataprepEditor::CreateDataprepLayout()
 {
-	return FTabManager::NewLayout("Standalone_DataprepEditor_Layout_v0.8")
+	return FTabManager::NewLayout("Standalone_DataprepEditor_Layout_v0.9")
 		->AddArea
 		(
 			FTabManager::NewPrimaryArea()->SetOrientation(Orient_Vertical)
@@ -994,7 +994,7 @@ TSharedRef<FTabManager::FLayout> FDataprepEditor::CreateDataprepLayout()
 						(
 							FTabManager::NewStack()
 							->SetSizeCoefficient(0.85f)
-							->AddTab(DataprepGraphEditorTabId, ETabState::ClosedTab)
+							->AddTab(DataprepGraphEditorTabId, ETabState::OpenedTab)
 							->SetHideTabWell( true )
 						)
 					)
