Dataprep filter preview improvement.
A preview can now be close and close it self automatically when the observed data is invalid.
Fixed some minor issue with columns of the scene preview and asset preview of the dataprep editor.
Ajusted the visual of the filter preview.

diff --git a/Engine/Plugins/Enterprise/DataprepEditor/Source/DataprepCore/Public/DataprepActionAsset.h b/Engine/Plugins/Enterprise/DataprepEditor/Source/DataprepCore/Public/DataprepActionAsset.h
--- a/Engine/Plugins/Enterprise/DataprepEditor/Source/DataprepCore/Public/DataprepActionAsset.h
+++ b/Engine/Plugins/Enterprise/DataprepEditor/Source/DataprepCore/Public/DataprepActionAsset.h
@@ -48,7 +48,7 @@ namespace DataprepActionAsset
 }
 
 UCLASS(Experimental)
-class UDataprepActionStep : public UObject
+class DATAPREPCORE_API UDataprepActionStep : public UObject
 {
 	GENERATED_BODY()
 
diff --git a/Engine/Plugins/Enterprise/DataprepEditor/Source/DataprepEditor/Private/DataprepEditor.cpp b/Engine/Plugins/Enterprise/DataprepEditor/Source/DataprepEditor/Private/DataprepEditor.cpp
--- a/Engine/Plugins/Enterprise/DataprepEditor/Source/DataprepEditor/Private/DataprepEditor.cpp
+++ b/Engine/Plugins/Enterprise/DataprepEditor/Source/DataprepEditor/Private/DataprepEditor.cpp
@@ -52,6 +52,7 @@
 #include "Misc/ScopedSlowTask.h"
 #include "Modules/ModuleManager.h"
 #include "ObjectTools.h"
+#include "SceneOutlinerModule.h"
 #include "SceneOutlinerPublicTypes.h"
 #include "ScopedTransaction.h"
 #include "StatsViewerModule.h"
@@ -1166,6 +1167,31 @@ void FDataprepEditor::OnActionsContextChanged( const UDataprepActionAsset* Actio
 	}
 }
 
+void FDataprepEditor::RefreshColumnsForPreviewSystem()
+{
+	if ( PreviewSystem->HasObservedObjects() )
+	{
+		SceneOutliner->RemoveColumn( SceneOutliner::FBuiltInColumnTypes::ActorInfo() );
+		SceneOutliner::FColumnInfo ColumnInfo( SceneOutliner::EColumnVisibility::Visible, 100, FCreateSceneOutlinerColumn::CreateLambda( [Preview = PreviewSystem](ISceneOutliner& InSceneOutliner) -> TSharedRef< ISceneOutlinerColumn >
+			{
+				return MakeShared<FDataprepPreviewOutlinerColumn>( InSceneOutliner, Preview );
+			} ) );
+		SceneOutliner->AddColumn( FDataprepPreviewOutlinerColumn::ColumnID, ColumnInfo );
+
+		AssetPreviewView->AddColumn( MakeShared<FDataprepPreviewAssetColumn>(PreviewSystem) );
+	}
+	else
+	{
+		SceneOutliner->RemoveColumn( FDataprepPreviewOutlinerColumn::ColumnID );
+		FSceneOutlinerModule& SceneOutlinerModule = FModuleManager::Get().LoadModuleChecked<FSceneOutlinerModule>("SceneOutliner");
+		SceneOutliner::FDefaultColumnInfo* ActorInfoColumPtr = SceneOutlinerModule.DefaultColumnMap.Find( SceneOutliner::FBuiltInColumnTypes::ActorInfo() );
+		check( ActorInfoColumPtr );
+		SceneOutliner->AddColumn( SceneOutliner::FBuiltInColumnTypes::ActorInfo(), ActorInfoColumPtr->ColumnInfo );
+
+		AssetPreviewView->RemoveColumn( FDataprepPreviewAssetColumn::ColumnID );
+	}
+}
+
 void FDataprepEditor::UpdateDataForPreviewSystem()
 {
 	TArray<UObject*> ObjectsForThePreviewSystem;
@@ -1183,24 +1209,32 @@ void FDataprepEditor::UpdateDataForPreviewSystem()
 
 bool FDataprepEditor::IsPreviewingStep(const UDataprepParameterizableObject* StepObject) const 
 {
-	return PreviewSystem->IsObservingObject( StepObject );
+	return PreviewSystem->IsObservingStepObject( StepObject );
+}
+
+int32 FDataprepEditor::GetCountOfPreviewedSteps() const
+{
+	return PreviewSystem->GetObservedStepsCount();
 }
 
 void FDataprepEditor::OnStepObjectsAboutToBeDeleted(const TArrayView<UDataprepParameterizableObject*>& StepObjects)
 {
-	// Todo Implement removal of objects from the preview system
+	for ( UDataprepParameterizableObject* StepObject : StepObjects )
+	{
+		if ( IsPreviewingStep( StepObject ) )
+		{
+			ClearPreviewedObjects();
+			break;
+		}
+	}
 }
 
 void FDataprepEditor::PostUndo(bool bSuccess)
 {
 	if ( bSuccess )
 	{
-		if ( PreviewSystem->HasObservedObjects() )
-		{
-			// Check if a object as disappear
-			PreviewSystem->RestartProcessing();
-			// Set Preview view
-		}
+		PreviewSystem->EnsureValidityOfTheObservedObjects();
+		RefreshColumnsForPreviewSystem();
 	}
 }
 
@@ -1208,33 +1242,13 @@ void FDataprepEditor::PostRedo(bool bSuccess)
 {
 	if ( bSuccess )
 	{
-		if ( PreviewSystem->HasObservedObjects() )
-		{
-			// Check if a object as disappear
-			PreviewSystem->RestartProcessing();
-			// Set Preview view
-		}
+		PreviewSystem->EnsureValidityOfTheObservedObjects();
+		RefreshColumnsForPreviewSystem();
 	}
 }
 
 void FDataprepEditor::SetPreviewedObjects(const TArrayView<UDataprepParameterizableObject*>& ObservedObjects)
 {
-	// (possible improvement) A validation that the observed object are from the same action in that the action is from the edited dataprep could help here.
-
-	if ( !PreviewSystem->HasObservedObjects() )
-	{ 
-		SceneOutliner->RemoveColumn( SceneOutliner::FBuiltInColumnTypes::ActorInfo() );
-		
-		SceneOutliner::FColumnInfo ColumnInfo( SceneOutliner::EColumnVisibility::Visible, 100, FCreateSceneOutlinerColumn::CreateLambda([Preview = PreviewSystem](ISceneOutliner& InSceneOutliner) -> TSharedRef< ISceneOutlinerColumn >
-			{
-				return MakeShared<FDataprepPreviewOutlinerColumn>( InSceneOutliner, Preview );
-			}));
-
-		SceneOutliner->AddColumn( FDataprepPreviewOutlinerColumn::ColumnID, ColumnInfo );
-
-		AssetPreviewView->AddColumn( MakeShared<FDataprepPreviewAssetColumn>( PreviewSystem ) );
-	}
-
 	PreviewSystem->SetObservedObjects( ObservedObjects );
 	
 	if ( SDataprepGraphEditor* RawGraphEditor = GraphEditor.Get() )
@@ -1242,6 +1256,13 @@ void FDataprepEditor::SetPreviewedObjects(const TArrayView<UDataprepParameteriza
 		// Refresh the graph
 		RawGraphEditor->NotifyGraphChanged();
 	}
+
+	RefreshColumnsForPreviewSystem();
+}
+
+void FDataprepEditor::ClearPreviewedObjects()
+{
+	SetPreviewedObjects( MakeArrayView<UDataprepParameterizableObject*>( nullptr, 0 ) );
 }
 
 void DataprepEditorUtil::DeleteTemporaryPackage( const FString& PathToDelete )
diff --git a/Engine/Plugins/Enterprise/DataprepEditor/Source/DataprepEditor/Private/DataprepEditor.h b/Engine/Plugins/Enterprise/DataprepEditor/Source/DataprepEditor/Private/DataprepEditor.h
--- a/Engine/Plugins/Enterprise/DataprepEditor/Source/DataprepEditor/Private/DataprepEditor.h
+++ b/Engine/Plugins/Enterprise/DataprepEditor/Source/DataprepEditor/Private/DataprepEditor.h
@@ -115,9 +115,15 @@ public:
 	/** Setup the preview system to observe those steps */
 	void SetPreviewedObjects(const TArrayView<UDataprepParameterizableObject*>& ObservedObjects);
 
+	/** Clear any ongoing preview from the step preview system for this editor */
+	void ClearPreviewedObjects();
+
 	/** Is the preview system observing this step */
 	bool IsPreviewingStep(const UDataprepParameterizableObject* StepObject) const;
 
+	/** Return the number of steps the preview system is currently previewing */
+	int32 GetCountOfPreviewedSteps() const;
+
 private:
 	void BindCommands();
 	void OnSaveScene();
@@ -197,6 +203,9 @@ private:
 	/** Handles change to the content passed to an action */
 	void OnActionsContextChanged( const UDataprepActionAsset* ActionAsset, bool bWorldChanged, bool bAssetsChanged, const TArray< TWeakObjectPtr<UObject> >& NewAssets );
 
+	/** Refresh the columns of the scene preview (outliner) and the asset preview */
+	void RefreshColumnsForPreviewSystem();
+
 	/** Update the data used by the preview system */
 	void UpdateDataForPreviewSystem();
 
diff --git a/Engine/Plugins/Enterprise/DataprepEditor/Source/DataprepEditor/Private/DataprepEditorStyle.cpp b/Engine/Plugins/Enterprise/DataprepEditor/Source/DataprepEditor/Private/DataprepEditorStyle.cpp
--- a/Engine/Plugins/Enterprise/DataprepEditor/Source/DataprepEditor/Private/DataprepEditorStyle.cpp
+++ b/Engine/Plugins/Enterprise/DataprepEditor/Source/DataprepEditor/Private/DataprepEditorStyle.cpp
@@ -128,7 +128,7 @@ void FDataprepEditorStyle::Initialize()
 		}
 
 
-		StyleSet->Set("Graph.ActionStepNode.PreviewColor", FLinearColor(1.f, 0.87f, 0.0117f, 1.f));
+		StyleSet->Set("Graph.ActionStepNode.PreviewColor", FLinearColor(0.822786f, 0.715693f, 0.0f, 1.f));
 
 		{
 			FTextBlockStyle TilteTextBlockStyle = FEditorStyle::GetWidgetStyle< FTextBlockStyle >("NormalText");
diff --git a/Engine/Plugins/Enterprise/DataprepEditor/Source/DataprepEditor/Private/DataprepGraph/DataprepGraphActionNode.h b/Engine/Plugins/Enterprise/DataprepEditor/Source/DataprepEditor/Private/DataprepGraph/DataprepGraphActionNode.h
--- a/Engine/Plugins/Enterprise/DataprepEditor/Source/DataprepEditor/Private/DataprepGraph/DataprepGraphActionNode.h
+++ b/Engine/Plugins/Enterprise/DataprepEditor/Source/DataprepEditor/Private/DataprepGraph/DataprepGraphActionNode.h
@@ -64,7 +64,7 @@ protected:
  * to an SGraphNode in order to display actions in a SDataprepGraphEditor.
  */
 UCLASS(MinimalAPI)
-class UDataprepGraphActionNode : public UEdGraphNode
+class UDataprepGraphActionNode final : public UEdGraphNode
 {
 	GENERATED_BODY()
 
diff --git a/Engine/Plugins/Enterprise/DataprepEditor/Source/DataprepEditor/Private/PreviewSystem/DataprepPreviewSystem.cpp b/Engine/Plugins/Enterprise/DataprepEditor/Source/DataprepEditor/Private/PreviewSystem/DataprepPreviewSystem.cpp
--- a/Engine/Plugins/Enterprise/DataprepEditor/Source/DataprepEditor/Private/PreviewSystem/DataprepPreviewSystem.cpp
+++ b/Engine/Plugins/Enterprise/DataprepEditor/Source/DataprepEditor/Private/PreviewSystem/DataprepPreviewSystem.cpp
@@ -2,7 +2,10 @@
 
 #include "PreviewSystem/DataprepPreviewSystem.h"
 
+#include "DataprepActionAsset.h"
+#include "DataprepAsset.h"
 #include "DataprepBindingCommandChange.h"
+#include "DataprepCoreUtils.h"
 #include "DataprepParameterizableObject.h"
 #include "SelectionSystem/DataprepFilter.h"
 
@@ -179,16 +182,45 @@ void FDataprepPreviewSystem::SetObservedObjects(const TArrayView<UDataprepParame
 {
 	StopTrackingObservedObjects();
 
-	ObservedObjects.Empty( StepObjects.Num() );
-
-	// Reserve the double for the filters which we need to track their also their fetcher
+	ObservedSteps.Empty( StepObjects.Num() );
+	// Reserve the double. For the filters we need to track their fetcher also.
 	ObservedOnPostEdit.Empty( StepObjects.Num() * 2 );
 
+	ObservedActionAsset = nullptr;
 	for ( UDataprepParameterizableObject* Object : StepObjects )
 	{
 		if ( Object )
 		{
-			ObservedObjects.Add( Object );
+			UDataprepActionAsset* ActionAssetOfStep = FDataprepCoreUtils::GetDataprepActionAssetOf( Object );
+			if ( !ObservedActionAsset )
+			{
+				ObservedActionAsset = ActionAssetOfStep;
+			}
+			
+			if ( !ActionAssetOfStep || ActionAssetOfStep != ObservedActionAsset )
+			{
+				ObservedActionAsset = nullptr;
+				ObservedSteps.Empty();
+				ObservedOnPostEdit.Empty();
+				RestartProcessing();
+				return;
+			}
+		}
+	}
+
+	if ( ObservedActionAsset )
+	{
+		OnActionStepOrderChangedHandle = ObservedActionAsset->GetOnStepsOrderChanged().AddLambda( [this]()
+			{
+				EnsureValidityOfTheObservedObjects();
+			});
+	}
+
+	for ( UDataprepParameterizableObject* Object : StepObjects )
+	{
+		if ( Object )
+		{
+			ObservedSteps.Add( Object );
 			
 			FDelegateHandle Handle = Object->GetOnPostEdit().AddSP( this, &FDataprepPreviewSystem::OnObservedObjectPostEdit );
 			ObservedOnPostEdit.Add( Object, Handle );
@@ -205,12 +237,16 @@ void FDataprepPreviewSystem::SetObservedObjects(const TArrayView<UDataprepParame
 		}
 	}
 
-	RestartProcessing();
+	// Make sure the step are inserted in the right order
+	EnsureValidityOfTheObservedObjects();
 }
 
 void FDataprepPreviewSystem::AddReferencedObjects(FReferenceCollector& Collector)
 {
-	Collector.AddReferencedObjects( ObservedObjects );
+	Collector.AddReferencedObjects( ObservedSteps );
+	Collector.AddReferencedObjects( ObservedOnPostEdit );
+	Collector.AddReferencedObjects( PreviewResult );
+	Collector.AddReferencedObject( ObservedActionAsset );
 }
 
 TSharedPtr<FDataprepPreviewProcessingResult> FDataprepPreviewSystem::GetPreviewDataForObject(UObject* Object) const
@@ -227,7 +263,7 @@ void FDataprepPreviewSystem::IncrementalProcess()
 {
 	const int32 ItemsCount = PreviewResult.Num();
 	int32 BudgetLeft = IncrementalCount;
-	while ( BudgetLeft > 0 && CurrentProgress.CurrentFilterIndex < ObservedObjects.Num() )
+	while ( BudgetLeft > 0 && CurrentProgress.CurrentFilterIndex < ObservedSteps.Num() )
 	{
 		int32 ItemToProcessCount = ItemsCount - CurrentProgress.CurrentObjectProcessed;
 
@@ -240,7 +276,7 @@ void FDataprepPreviewSystem::IncrementalProcess()
 
 		if ( CurrentProgress.Iterator.IsValid() )
 		{
-			if ( UDataprepFilter* Filter = Cast<UDataprepFilter>( ObservedObjects[CurrentProgress.CurrentFilterIndex] ) )
+			if ( UDataprepFilter* Filter = Cast<UDataprepFilter>( ObservedSteps[CurrentProgress.CurrentFilterIndex] ) )
 			{
 				bDidSomeProcessing = true;
 				PrepareFilterBuffers( IncrementalCount );
@@ -298,7 +334,7 @@ void FDataprepPreviewSystem::RestartProcessing()
 			}
 		}
 
-		if ( ObservedObjects.Num() > 0 )
+		if ( ObservedSteps.Num() > 0 )
 		{
 			bIsProcessing = true;
 		}
@@ -323,11 +359,87 @@ bool FDataprepPreviewSystem::HasAnObjectObserved(const TArrayView<UDataprepParam
 	return false;
 }
 
-bool FDataprepPreviewSystem::IsObservingObject(const UDataprepParameterizableObject* StepObject) const
+bool FDataprepPreviewSystem::IsObservingStepObject(const UDataprepParameterizableObject* StepObject) const
 {
 	return ObservedOnPostEdit.Contains( StepObject );
 }
 
+void FDataprepPreviewSystem::EnsureValidityOfTheObservedObjects()
+{
+	if ( ObservedActionAsset )
+	{
+		if ( UDataprepAsset* Owner = Cast<UDataprepAsset>( ObservedActionAsset->GetOuter() ) )
+		{
+			bool bActionWasFoundInParent = false;
+			for ( int32 Index = 0; Index < Owner->GetActionCount(); ++Index )
+			{
+				if ( ObservedActionAsset == Owner->GetAction( Index ) )
+				{
+					bActionWasFoundInParent = true;
+					break;
+				}
+			}
+
+			if ( !bActionWasFoundInParent )
+			{
+				// The action wasn't found in the dataprep asset. Cancel the preview.
+				SetObservedObjects( MakeArrayView<UDataprepParameterizableObject*>( nullptr, 0 ) );
+				return;
+			}
+		}
+		else
+		{
+			// If the action is not own by a dataprep asset bail out.
+			SetObservedObjects( MakeArrayView<UDataprepParameterizableObject*>( nullptr, 0 ) );
+			return;
+		}
+
+		TArray<UDataprepParameterizableObject*> ObjectOrdering;
+		ObjectOrdering.Reserve( ObservedSteps.Num() );
+
+		TSet<UDataprepParameterizableObject*> ObservedStepsSet;
+		ObservedStepsSet.Append( ObservedSteps );
+
+		bool bShouldUpdateStepsOrder = false;
+
+		int32 StepCount = ObservedActionAsset->GetStepsCount();
+		for ( int32 Index = 0; Index < StepCount; ++Index )
+		{
+			if ( UDataprepActionStep* Step = ObservedActionAsset->GetStep( Index ).Get() )
+			{
+				UDataprepParameterizableObject* StepObject = Step->GetStepObject();
+				if ( ObservedStepsSet.Contains( StepObject ) )
+				{
+					if ( ObservedSteps[ObjectOrdering.Num()] != StepObject )
+					{
+						bShouldUpdateStepsOrder = true;
+					}
+
+					ObjectOrdering.Add( StepObject );
+				}
+			}
+		}
+
+		if ( ObjectOrdering.Num() != ObservedSteps.Num() )
+		{
+			// We lost a step. Cancel the preview.
+			SetObservedObjects( MakeArrayView<UDataprepParameterizableObject*>( nullptr, 0 ) );
+			return;
+		}
+
+		if ( bShouldUpdateStepsOrder )
+		{
+			for ( int32 Index = 0; Index < ObjectOrdering.Num(); ++Index )
+			{
+				ObservedSteps[Index] = ObjectOrdering[Index];
+			}
+		}
+		
+		// Just in case a observed object was edited
+		RestartProcessing();
+	}
+}
+
 int32 FDataprepPreviewSystem::FillObjectsBuffer(int32 MaximunNumberOfObject)
 {
 	MaximunNumberOfObject = FMath::Min( ObjectsBuffer.Num(), MaximunNumberOfObject );
@@ -370,7 +482,7 @@ void FDataprepPreviewSystem::PopulateResultFromFilter(int32 NumberOfValidObjects
 
 	NumberOfValidObjects = FMath::Min( ObjectsBuffer.Num(), NumberOfValidObjects );
 	
-	const bool bIsLastFilter =  CurrentProgress.CurrentFilterIndex == ObservedObjects.Num() - 1;
+	const bool bIsLastFilter =  CurrentProgress.CurrentFilterIndex == ObservedSteps.Num() - 1;
 
 	ParallelFor( NumberOfValidObjects, [this, bIsLastFilter](int32 Index)
 		{
@@ -408,6 +520,11 @@ void FDataprepPreviewSystem::PopulateResultFromFilter(int32 NumberOfValidObjects
 
 void FDataprepPreviewSystem::StopTrackingObservedObjects()
 {
+	if ( ObservedActionAsset )
+	{
+		ObservedActionAsset->GetOnStepsOrderChanged().Remove( OnActionStepOrderChangedHandle );
+	}
+
 	for ( auto& Pair : ObservedOnPostEdit )
 	{
 		if ( Pair.Key )
@@ -427,7 +544,7 @@ void FDataprepPreviewSystem::Tick(float DeltaTime)
 {
 	IncrementalProcess();
 
-	if ( CurrentProgress.CurrentFilterIndex >= ObservedObjects.Num() )
+	if ( CurrentProgress.CurrentFilterIndex >= ObservedSteps.Num() )
 	{
 		bIsProcessing = false;
 		OnPreviewIsDoneProcessing.Broadcast();
diff --git a/Engine/Plugins/Enterprise/DataprepEditor/Source/DataprepEditor/Private/PreviewSystem/DataprepPreviewSystem.h b/Engine/Plugins/Enterprise/DataprepEditor/Source/DataprepEditor/Private/PreviewSystem/DataprepPreviewSystem.h
--- a/Engine/Plugins/Enterprise/DataprepEditor/Source/DataprepEditor/Private/PreviewSystem/DataprepPreviewSystem.h
+++ b/Engine/Plugins/Enterprise/DataprepEditor/Source/DataprepEditor/Private/PreviewSystem/DataprepPreviewSystem.h
@@ -14,6 +14,7 @@
 #include "UObject/GCObject.h"
 
 class FReferenceCollector;
+class UDataprepActionAsset;
 
 struct TStatId;
 
@@ -71,7 +72,7 @@ public:
 
 	void SetObservedObjects(const TArrayView<UDataprepParameterizableObject*>& Objects);
 
-	bool HasObservedObjects() const { return ObservedObjects.Num() > 0; }
+	bool HasObservedObjects() const { return ObservedSteps.Num() > 0; }
 
 	// Begin FGCObject interface
 	virtual void AddReferencedObjects(FReferenceCollector& Collector) override;
@@ -87,8 +88,17 @@ public:
 
 	bool HasAnObjectObserved(const TArrayView<UDataprepParameterizableObject*>& StepObjects) const;
 
-	// Is the preview system observing an object
-	bool IsObservingObject(const UDataprepParameterizableObject* StepObject) const;
+	// Is the preview system observing the object
+	bool IsObservingStepObject(const UDataprepParameterizableObject* StepObject) const;
+
+	// Return the current number of the step observed
+	int32 GetObservedStepsCount() const { return ObservedSteps.Num(); }
+
+	/**
+	 * Check if the observed object are valid and also check if the observed object are in the right order
+	 * If an observed object is invalid clear the preview
+	 */
+	void EnsureValidityOfTheObservedObjects();
 
 	// Broadcast when the preview system is done processing the data to preview
 	DECLARE_EVENT(FDataprepPreviewSystem, FOnPreviewIsDoneProcessing)
@@ -129,8 +139,12 @@ public:
 	// End FTickableObjectBase interface
 
 private:
+	// The action from which the step object are previewed
+	UDataprepActionAsset* ObservedActionAsset;
+	FDelegateHandle OnActionStepOrderChangedHandle;
+
 	// The filters previewed
-	TArray<UDataprepParameterizableObject*> ObservedObjects;
+	TArray<UDataprepParameterizableObject*> ObservedSteps;
 
 	// The handle for the PostEditEvents of the observed objects
 	TMap<UDataprepParameterizableObject*,FDelegateHandle> ObservedOnPostEdit;
diff --git a/Engine/Plugins/Enterprise/DataprepEditor/Source/DataprepEditor/Private/Widgets/DataprepGraph/SDataprepGraphEditor.cpp b/Engine/Plugins/Enterprise/DataprepEditor/Source/DataprepEditor/Private/Widgets/DataprepGraph/SDataprepGraphEditor.cpp
--- a/Engine/Plugins/Enterprise/DataprepEditor/Source/DataprepEditor/Private/Widgets/DataprepGraph/SDataprepGraphEditor.cpp
+++ b/Engine/Plugins/Enterprise/DataprepEditor/Source/DataprepEditor/Private/Widgets/DataprepGraph/SDataprepGraphEditor.cpp
@@ -48,6 +48,49 @@ const float SDataprepGraphEditor::HorizontalPadding = 20.f;
 
 TSharedPtr<FDataprepGraphEditorNodeFactory> SDataprepGraphEditor::NodeFactory;
 
+namespace DataprepGraphEditorUtils
+{
+	void ForEachActionAndStep(const TSet<UObject*>& Nodes,  TFunctionRef<bool (UDataprepActionAsset*, bool /* bIsFromActionNode */)> OnEachAction, TFunctionRef<bool (UDataprepParameterizableObject*)> OnEachStep)
+	{
+		for (UObject* Node : Nodes)
+		{
+			if ( Node )
+			{
+				if (UDataprepGraphActionNode* ActionNode = Cast<UDataprepGraphActionNode>(Node))
+				{
+					if ( OnEachAction( ActionNode->GetDataprepActionAsset(), true ) )
+					{
+						break;
+					}
+				}
+				else if (UDataprepGraphActionStepNode* StepNode = Cast<UDataprepGraphActionStepNode>(Node) )
+				{
+					if (UDataprepActionAsset* Action = StepNode->GetDataprepActionAsset())
+					{
+						if ( OnEachAction( Action, false ) )
+						{
+							break;
+						}
+
+						const int32 StepIndex = StepNode->GetStepIndex();
+						if (UDataprepActionStep* Step = Action->GetStep(StepIndex).Get())
+						{
+							if (UDataprepParameterizableObject* StepObject = Step->GetStepObject())
+							{
+								if ( OnEachStep( StepObject ) )
+								{
+									break;
+								}
+							}
+						}
+					}
+				}
+			}
+		}
+	}
+}
+
+
 TSharedPtr<SGraphNode> FDataprepGraphNodeFactory::CreateNodeWidget(UEdGraphNode* InNode)
 {
 	if (UDataprepGraphRecipeNode* RecipeNode = Cast<UDataprepGraphRecipeNode>(InNode))
@@ -758,37 +801,39 @@ FActionMenuContent SDataprepGraphEditor::OnCreateNodeOrPinMenu(UEdGraph* Current
 			TArray<UDataprepFilter*> Filters;
 			Filters.Reserve( SelectedNodes.Num() );
 
-			const UDataprepActionAsset* ClikedAction = FirstStepNode->GetDataprepActionAsset();
-
+			const UDataprepActionAsset* ClickedAction = FirstStepNode->GetDataprepActionAsset();
 
-			for ( UObject* Node : SelectedNodes )
-			{
-				if ( UDataprepGraphActionStepNode* StepNode = Cast<UDataprepGraphActionStepNode>( Node ) )
+			// Check the if the selection is all from the same action and only from step node
+			auto OnEachAction = [&bAreFilterFromSameAction, &bIsSelectionOnlyFilters, ClickedAction] (UDataprepActionAsset* Action, bool bIsFromActionNode) -> bool
 				{
-					if ( UDataprepActionAsset* Action = StepNode->GetDataprepActionAsset() )
+					if ( bIsFromActionNode )
 					{
-						bAreFilterFromSameAction &= Action == ClikedAction;
-
-						int32 StepIndex = StepNode->GetStepIndex();
-						if ( UDataprepActionStep* Step = Action->GetStep( StepIndex ).Get() )
-						{
-							if ( UDataprepParameterizableObject* StepObject = Step->GetStepObject() )
-							{
-								if ( UDataprepFilter* Filter = Cast<UDataprepFilter>( StepObject ) )
-								{
-									Filters.Add( Filter );
-									continue;
-								}
-							}
-						}
+						bIsSelectionOnlyFilters = false;
+						return true;
 					}
-				}
+					else
+					{
+						bAreFilterFromSameAction &= Action == ClickedAction;
+					}
+					return !bAreFilterFromSameAction;
+				};
 
-				bIsSelectionOnlyFilters = false;
-				break;
-			}
+			// Check if the selection is only filters
+			auto OnEachStepForFilterOnly = [&bIsSelectionOnlyFilters, &Filters](UDataprepParameterizableObject* StepObject) -> bool
+				{
+					if ( UDataprepFilter* Filter = Cast<UDataprepFilter>( StepObject ) )
+					{
+						Filters.Add( Filter );
+					}
+					else
+					{
+						bIsSelectionOnlyFilters = false;
+					}
 
+					return !bIsSelectionOnlyFilters;
+				};
 
+			DataprepGraphEditorUtils::ForEachActionAndStep( SelectedNodes, OnEachAction, OnEachStepForFilterOnly );
 
 			if ( bIsSelectionOnlyFilters )
 			{
@@ -810,18 +855,57 @@ FActionMenuContent SDataprepGraphEditor::OnCreateNodeOrPinMenu(UEdGraph* Current
 
 					if ( TSharedPtr<FDataprepEditor> DataprepEditorPtr = DataprepEditor.Pin() )
 					{ 
-						FUIAction SetFilterPreview;
-						SetFilterPreview.CanExecuteAction.BindLambda( [bAreFilterFromSameAction]() { return bAreFilterFromSameAction; } );
-						SetFilterPreview.ExecuteAction.BindLambda( [Filters, DataprepEditorPtr]()
+
+						bool bAreAllFilterPreviewed = true;
+
+						/**
+						 * Check if the filters are the same as the previewed steps from the preview system
+						 * Work with the assumption that the filters array contains only unique objects
+						 */
+						{
+							for ( UDataprepFilter* Filter : Filters )
+							{
+								if ( !DataprepEditorPtr->IsPreviewingStep(Filter) )
+								{
+									bAreAllFilterPreviewed = false;
+									break;
+								}
+							}
+							if ( Filters.Num() != DataprepEditorPtr->GetCountOfPreviewedSteps() )
 							{
-								TArray<UDataprepParameterizableObject*> ObjectsToObserve( Filters );
-								DataprepEditorPtr->SetPreviewedObjects( MakeArrayView<UDataprepParameterizableObject*>( ObjectsToObserve.GetData(), ObjectsToObserve.Num() ) );
-							});
-
-						MenuBuilder->AddMenuEntry(LOCTEXT("SetFilterPreview", "Preview Filter(s)"),
-							bAreFilterFromSameAction ? LOCTEXT("SetFilterPreviewTooltip", "Change the columns of the scene preview and asset preview tabs to display a preview of what the filters would select from the current scene") : LOCTEXT("SetFilterPreviewFailTooltip", "The filters can only be previewed if they are from the same action."),
-							FSlateIcon(),
-							SetFilterPreview);
+								bAreAllFilterPreviewed = false;
+							}
+						}
+
+
+						if ( bAreAllFilterPreviewed )
+						{
+							FUIAction ClearFilterPreview;
+							ClearFilterPreview.ExecuteAction.BindLambda([DataprepEditorPtr]()
+								{
+									DataprepEditorPtr->ClearPreviewedObjects();
+								});
+
+							MenuBuilder->AddMenuEntry(LOCTEXT("ClearFilterPreview", "Clear the previewed Filter(s)"),
+								LOCTEXT("ClearFilterPreviewTooltip", "Clear the columns of the scene preview and asset preview tabs of the Filter Preview."),
+								FSlateIcon(),
+								ClearFilterPreview);
+						}
+						else
+						{
+							FUIAction SetFilterPreview;
+							SetFilterPreview.CanExecuteAction.BindLambda( [bAreFilterFromSameAction]() { return bAreFilterFromSameAction; } );
+							SetFilterPreview.ExecuteAction.BindLambda( [Filters, DataprepEditorPtr]()
+								{
+									TArray<UDataprepParameterizableObject*> ObjectsToObserve( Filters );
+									DataprepEditorPtr->SetPreviewedObjects( MakeArrayView<UDataprepParameterizableObject*>( ObjectsToObserve.GetData(), ObjectsToObserve.Num() ) );
+								});
+
+							MenuBuilder->AddMenuEntry(LOCTEXT("SetFilterPreview", "Preview Filter(s)"),
+								bAreFilterFromSameAction ? LOCTEXT("SetFilterPreviewTooltip", "Change the columns of the scene preview and asset preview tabs to display a preview of what the filters would select from the current scene") : LOCTEXT("SetFilterPreviewFailTooltip", "The filters can only be previewed if they are from the same action."),
+								FSlateIcon(),
+								SetFilterPreview);
+						}
 					}
 				}
 				MenuBuilder->EndSection();
diff --git a/Engine/Plugins/Enterprise/DataprepEditor/Source/DataprepEditor/Private/Widgets/SAssetsPreviewWidget.cpp b/Engine/Plugins/Enterprise/DataprepEditor/Source/DataprepEditor/Private/Widgets/SAssetsPreviewWidget.cpp
--- a/Engine/Plugins/Enterprise/DataprepEditor/Source/DataprepEditor/Private/Widgets/SAssetsPreviewWidget.cpp
+++ b/Engine/Plugins/Enterprise/DataprepEditor/Source/DataprepEditor/Private/Widgets/SAssetsPreviewWidget.cpp
@@ -432,12 +432,18 @@ namespace AssetPreviewWidget
 
 	void SAssetsPreviewWidget::AddColumn(TSharedRef<IAssetPreviewColumn> Column)
 	{
-		PendingColumnsToAdd.Add( Column );
+		if ( !Columns.Contains( Column->GetColumnID() ) )
+		{
+			PendingColumnsToAdd.Add( Column );
+		}
 	}
 
 	void SAssetsPreviewWidget::RemoveColumn(FName ColumnID)
 	{
-		PendingColumnsToRemove.Add( ColumnID );
+		if ( Columns.Contains( ColumnID ) )
+		{
+			PendingColumnsToRemove.Add( ColumnID );
+		}
 	}
 
 	void SAssetsPreviewWidget::OnSetExpansionRecursive(IAssetTreeItemPtr InTreeNode, bool bInIsItemExpanded)
@@ -522,7 +528,6 @@ namespace AssetPreviewWidget
 			PendingColumnsToRemove.Empty();
 		}
 
-
 		if ( PendingColumnsToAdd.Num() > 0 ) 
 		{
 			PendingColumnsToAdd.StableSort([](const TSharedRef<IAssetPreviewColumn>& First, const TSharedRef<IAssetPreviewColumn>& Second) -> bool
@@ -538,24 +543,29 @@ namespace AssetPreviewWidget
 			TArray<int32> InsertionIndex;
 			InsertionIndex.Reserve( PendingColumnsToAdd.Num() );
 
-			for (const auto& Pair : Columns)
+
+			for (const TPair<FName, TSharedRef<IAssetPreviewColumn>>& Pair : Columns)
 			{
-				if ( Pair.Value->GetCulumnPositionPriorityIndex() < CurrentPriority )
+				while ( Pair.Value->GetCulumnPositionPriorityIndex() < CurrentPriority )
 				{
 					InsertionIndex.Add( CurrentElementToAddIndex + CurrentColumnIndex );
 					++CurrentElementToAddIndex;
-					if ( CurrentElementToAddIndex >= PendingColumnsToAdd.Num() )
+					if ( CurrentElementToAddIndex == PendingColumnsToAdd.Num() )
 					{
 						break;
 					}
+				
 					CurrentPriority = PendingColumnsToAdd[CurrentElementToAddIndex]->GetCulumnPositionPriorityIndex();
 				}
 
+				if ( CurrentElementToAddIndex == PendingColumnsToAdd.Num() )
+				{
+					break;
+				}
+
 				++CurrentColumnIndex;
 			}
 
-			
-
 			Columns.Reserve( PendingColumnsToAdd.Num() );
 			for ( TSharedRef<IAssetPreviewColumn>& NewColumn : PendingColumnsToAdd )
 			{
diff --git a/Engine/Plugins/Enterprise/DataprepEditor/Source/DataprepEditor/Private/Widgets/SDataprepPreviewRow.cpp b/Engine/Plugins/Enterprise/DataprepEditor/Source/DataprepEditor/Private/Widgets/SDataprepPreviewRow.cpp
--- a/Engine/Plugins/Enterprise/DataprepEditor/Source/DataprepEditor/Private/Widgets/SDataprepPreviewRow.cpp
+++ b/Engine/Plugins/Enterprise/DataprepEditor/Source/DataprepEditor/Private/Widgets/SDataprepPreviewRow.cpp
@@ -8,6 +8,7 @@
 #include "EditorFontGlyphs.h"
 #include "EditorStyleSet.h"
 #include "Internationalization/Text.h"
+#include "Styling/SlateColor.h"
 #include "Widgets/Layout/SBox.h"
 #include "Widgets/SBoxPanel.h"
 #include "Widgets/Text/STextBlock.h"
@@ -16,22 +17,21 @@ void SDataprepPreviewRow::Construct(const FArguments& InArgs, const TSharedPtr<F
 {
 	PreviewData = InPreviewData;
 
+	FSlateFontInfo AwsomeFontStyle = FEditorStyle::Get().GetFontStyle("FontAwesome.11");
+	AwsomeFontStyle.Size =9 ;
+
 	ChildSlot
 	[
 		SNew( SHorizontalBox )
 		+ SHorizontalBox::Slot()
 		.VAlign( VAlign_Center )
 		.HAlign( HAlign_Left )
-		.AutoWidth()
+		.MaxWidth(15.f)
 		[
-			SNew( SBox )
-			.WidthOverride( 18.f )
-			[
 				SNew( STextBlock )
-				.Font( FEditorStyle::Get().GetFontStyle("FontAwesome.11") )
+				.Font( AwsomeFontStyle )
 				.ColorAndOpacity( FDataprepEditorStyle::GetColor("Graph.ActionStepNode.PreviewColor"))
 				.Text( this, &SDataprepPreviewRow::GetIcon )
-			]
 		]
 		+ SHorizontalBox::Slot()
 		.VAlign(VAlign_Center)
@@ -40,6 +40,7 @@ void SDataprepPreviewRow::Construct(const FArguments& InArgs, const TSharedPtr<F
 			SNew( STextBlock )
 			.Text( this, &SDataprepPreviewRow::GetLabel )
 			.HighlightText( InArgs._HighlightText )
+			.ColorAndOpacity( this, &SDataprepPreviewRow::GetTextColor )
 		]
 	];
 }
@@ -80,3 +81,15 @@ FText SDataprepPreviewRow::GetLabel() const
 	return {};
 }
 
+FSlateColor SDataprepPreviewRow::GetTextColor() const
+{
+	if ( FDataprepPreviewProcessingResult* Result = PreviewData.Get() )
+	{
+		if ( Result->Status == EDataprepPreviewStatus::Pass )
+		{
+			return FSlateColor::UseForeground();
+		}
+	}
+
+	return FSlateColor::UseSubduedForeground();
+}
diff --git a/Engine/Plugins/Enterprise/DataprepEditor/Source/DataprepEditor/Private/Widgets/SDataprepPreviewRow.h b/Engine/Plugins/Enterprise/DataprepEditor/Source/DataprepEditor/Private/Widgets/SDataprepPreviewRow.h
--- a/Engine/Plugins/Enterprise/DataprepEditor/Source/DataprepEditor/Private/Widgets/SDataprepPreviewRow.h
+++ b/Engine/Plugins/Enterprise/DataprepEditor/Source/DataprepEditor/Private/Widgets/SDataprepPreviewRow.h
@@ -25,5 +25,7 @@ private:
 
 	FText GetLabel() const;
 
+	FSlateColor GetTextColor() const;
+
 	TSharedPtr<FDataprepPreviewProcessingResult> PreviewData;
 };
\ No newline at end of file
diff --git a/Engine/Source/Editor/SceneOutliner/Private/SSceneOutliner.cpp b/Engine/Source/Editor/SceneOutliner/Private/SSceneOutliner.cpp
--- a/Engine/Source/Editor/SceneOutliner/Private/SSceneOutliner.cpp
+++ b/Engine/Source/Editor/SceneOutliner/Private/SSceneOutliner.cpp
@@ -1913,14 +1913,20 @@ namespace SceneOutliner
 
 	void SSceneOutliner::AddColumn(FName ColumId, const SceneOutliner::FColumnInfo& ColumInfo)
 	{
-		SharedData->ColumnMap.Add(ColumId, ColumInfo);
-		RefreshColums();
+		if (!SharedData->ColumnMap.Contains(ColumId))
+		{
+			SharedData->ColumnMap.Add(ColumId, ColumInfo);
+			RefreshColums();
+		}
 	}
 
 	void SSceneOutliner::RemoveColumn(FName ColumId)
 	{
-		SharedData->ColumnMap.Remove(ColumId);
-		RefreshColums();
+		if (SharedData->ColumnMap.Contains(ColumId))
+		{
+			SharedData->ColumnMap.Remove(ColumId);
+			RefreshColums();
+		}
 	}
 
 	TArray<FName> SSceneOutliner::GetColumnIds() const
diff --git a/Engine/Source/Editor/SceneOutliner/Public/SceneOutlinerPublicTypes.h b/Engine/Source/Editor/SceneOutliner/Public/SceneOutlinerPublicTypes.h
--- a/Engine/Source/Editor/SceneOutliner/Public/SceneOutlinerPublicTypes.h
+++ b/Engine/Source/Editor/SceneOutliner/Public/SceneOutlinerPublicTypes.h
@@ -89,7 +89,7 @@ namespace SceneOutliner
 	};
 
 	/** Default column information for the scene outliner */
-	struct FDefaultColumnInfo : public FColumnInfo
+	struct FDefaultColumnInfo
 	{
 		FDefaultColumnInfo(const FColumnInfo& InColumnInfo, TOptional<ESceneOutlinerMode> InValidMode = TOptional<ESceneOutlinerMode>())
 			: ColumnInfo(InColumnInfo), ValidMode(InValidMode)
