Fixed warning when opening Dataprep asset editor

#rnx

diff --git a/Engine/Plugins/Enterprise/DataprepEditor/Source/DataprepEditor/Private/Widgets/DataprepGraph/SDataprepGraphTrackNode.cpp b/Engine/Plugins/Enterprise/DataprepEditor/Source/DataprepEditor/Private/Widgets/DataprepGraph/SDataprepGraphTrackNode.cpp
--- a/Engine/Plugins/Enterprise/DataprepEditor/Source/DataprepEditor/Private/Widgets/DataprepGraph/SDataprepGraphTrackNode.cpp
+++ b/Engine/Plugins/Enterprise/DataprepEditor/Source/DataprepEditor/Private/Widgets/DataprepGraph/SDataprepGraphTrackNode.cpp
@@ -87,7 +87,7 @@ public:
 				[
 					SNew(STextBlock)
 					.Text(LOCTEXT("DataprepEmptyActionLabel", "+ Add Action"))
-					.TextStyle( &FDataprepEditorStyle::GetWidgetStyle<FTextBlockStyle>( "DataprepAction.TitleTextBlockStyle" ) )
+					.TextStyle( &FDataprepEditorStyle::GetWidgetStyle<FTextBlockStyle>( "DataprepAction.TitleTextStyle" ) )
 					.ColorAndOpacity(FDataprepEditorStyle::GetColor("DataprepAction.EmptyStep.Text.Hovered"))
 					.Justification(ETextJustify::Center)
 				]
