Fixed warning about static vs non-static actor when opening Dataprep editor's 3D viewport

#rnx

diff --git a/Engine/Plugins/Enterprise/DataprepEditor/Source/DataprepEditor/Private/Widgets/SDataprepEditorViewport.cpp b/Engine/Plugins/Enterprise/DataprepEditor/Source/DataprepEditor/Private/Widgets/SDataprepEditorViewport.cpp
--- a/Engine/Plugins/Enterprise/DataprepEditor/Source/DataprepEditor/Private/Widgets/SDataprepEditorViewport.cpp
+++ b/Engine/Plugins/Enterprise/DataprepEditor/Source/DataprepEditor/Private/Widgets/SDataprepEditorViewport.cpp
@@ -255,13 +255,41 @@ void SDataprepEditorViewport::Construct(const FArguments& InArgs, TSharedPtr<FDa
 		PreviewSceneWorld->ChangeFeatureLevel(NewFeatureLevel);
 	});
 
-	// Create actor in preview world to hold all preview mesh components
-	PreviewActor = TWeakObjectPtr<AActor>( PreviewSceneWorld->SpawnActor<AActor>( AActor::StaticClass(), FTransform::Identity ) );
+	TFunctionRef<AActor*(USceneComponent*, const TCHAR*)> CreateActor =
+		[&PreviewSceneWorld](USceneComponent* ParentComponent, const TCHAR* ActorName)-> AActor*
+		{
+			AActor* Actor = PreviewSceneWorld->SpawnActor<AActor>( AActor::StaticClass(), FTransform::Identity );
+
+			if ( Actor->GetRootComponent() == nullptr )
+			{
+				USceneComponent* RootComponent = NewObject< USceneComponent >( Actor, USceneComponent::StaticClass(), ActorName, RF_NoFlags );
+				Actor->SetRootComponent( RootComponent );
+			}
+
+			if(ParentComponent)
+			{
+				Actor->GetRootComponent()->AttachToComponent( ParentComponent, FAttachmentTransformRules::KeepRelativeTransform );
+			}
+
+			Actor->RegisterAllComponents();
 
-	if ( PreviewActor->GetRootComponent() == nullptr )
+			return Actor;
+		};
+
+	// Create root actor of 3D viewport's world
+	AActor* RootActor = CreateActor(nullptr, TEXT("RootPreviewActor"));
+	RootActorComponent = RootActor->GetRootComponent();
+
+
+	// Create actor to hold all 'Static' preview mesh components
+	StaticActor = CreateActor(RootActorComponent.Get(), TEXT("StaticPreviewActor"));
+	StaticActor->GetRootComponent()->SetMobility(EComponentMobility::Static);
+
+	// If HW supports it, create actor to hold all 'Movable' preview mesh components
+	if(GEditor->PreviewPlatform.GetEffectivePreviewFeatureLevel() > ERHIFeatureLevel::ES3_1)
 	{
-		USceneComponent* RootComponent = NewObject< USceneComponent >( PreviewActor.Get(), USceneComponent::StaticClass(), TEXT("PreviewActor"), RF_NoFlags );
-		PreviewActor->SetRootComponent( RootComponent );
+		MovableActor = CreateActor(RootActorComponent.Get(), TEXT("MovablePreviewActor"));
+		MovableActor->GetRootComponent()->SetMobility(EComponentMobility::Movable);
 	}
 
 	WorldToPreview = InArgs._WorldToPreview;
@@ -365,7 +393,7 @@ void SDataprepEditorViewport::UpdateScene()
 			SceneBounds.Min *= SceneUniformScale;
 
 			// Set uniform scale on root actor's root component
-			PreviewActor->GetRootComponent()->SetRelativeTransform( FTransform( FRotator::ZeroRotator, FVector::ZeroVector, FVector( SceneUniformScale ) ) );
+			RootActorComponent->SetRelativeTransform( FTransform( FRotator::ZeroRotator, FVector::ZeroVector, FVector( SceneUniformScale ) ) );
 
 			const int32 PerMeshColorsCount = sizeof(PerMeshColor);
 			int32 PerMeshColorIndex = 0;
@@ -381,16 +409,17 @@ void SDataprepEditorViewport::UpdateScene()
 						const FText Message = FText::Format( LOCTEXT("UpdateMeshes_AddOneComponent", "Adding {0} ..."), FText::FromString( SceneMeshComponent->GetOwner()->GetActorLabel() ) );
 						SubSlowTask.EnterProgressFrame( 1.0f, Message );
 
-						UCustomStaticMeshComponent* PreviewMeshComponent = NewObject< UCustomStaticMeshComponent >( PreviewActor.Get(), NAME_None, RF_Transient );
-						if (GEditor->PreviewPlatform.GetEffectivePreviewFeatureLevel() <= ERHIFeatureLevel::ES3_1)
+						EComponentMobility::Type Mobility = EComponentMobility::Static;
+						if(MovableActor.IsValid())
 						{
-							PreviewMeshComponent->SetMobility(EComponentMobility::Static);
-						}
-						else
-						{
-							PreviewMeshComponent->SetMobility( SceneMeshComponent->Mobility );
+							Mobility = SceneMeshComponent->Mobility;
 						}
 
+						AActor* PreviewActor = Mobility == EComponentMobility::Movable ? MovableActor.Get() : StaticActor.Get();
+
+						UCustomStaticMeshComponent* PreviewMeshComponent = NewObject< UCustomStaticMeshComponent >( PreviewActor, NAME_None, RF_Transient );
+						PreviewMeshComponent->SetMobility( Mobility );
+
 						PreviewMeshComponent->bOverrideLightMapRes = SceneMeshComponent->bOverrideLightMapRes;
 						PreviewMeshComponent->OverriddenLightMapRes = SceneMeshComponent->OverriddenLightMapRes;
 						PreviewMeshComponent->bForceWireframe = bWireframeRenderingMode;
diff --git a/Engine/Plugins/Enterprise/DataprepEditor/Source/DataprepEditor/Private/Widgets/SDataprepEditorViewport.h b/Engine/Plugins/Enterprise/DataprepEditor/Source/DataprepEditor/Private/Widgets/SDataprepEditorViewport.h
--- a/Engine/Plugins/Enterprise/DataprepEditor/Source/DataprepEditor/Private/Widgets/SDataprepEditorViewport.h
+++ b/Engine/Plugins/Enterprise/DataprepEditor/Source/DataprepEditor/Private/Widgets/SDataprepEditorViewport.h
@@ -317,8 +317,14 @@ private:
 	/** World containing the data to preview */
 	UWorld* WorldToPreview;
 
-	/** Root actor all mesh components will be attached to for preview. */
-	TWeakObjectPtr<AActor> PreviewActor;
+	/** Root component of the viewport world's root actor. */
+	TWeakObjectPtr<USceneComponent> RootActorComponent;
+
+	/** Actor to hold onto all 'Movable' preview mesh components */
+	TWeakObjectPtr<AActor> MovableActor;
+
+	/** Actor to hold onto all 'Static' preview mesh components */
+	TWeakObjectPtr<AActor> StaticActor;
 
 	/** Array of static mesh components added to the preview scene for display. */
 	TArray< TWeakObjectPtr< UStaticMeshComponent > > PreviewMeshComponents;
