Fixed issue with IES light texture disappearing on subsequent executes
  - root cause: A package in the transient folder used to execute the Dataprep asset was flagged as transient but not deleted
  - solution: Strengthened the cleanup done after executing the Dataprep asset. Add this cleanup to the execution of a Dataprep asset thru the execute menu entry

#rnx

diff --git a/Engine/Plugins/Enterprise/DataprepEditor/Source/DataprepCore/Private/DataprepCoreUtils.cpp b/Engine/Plugins/Enterprise/DataprepEditor/Source/DataprepCore/Private/DataprepCoreUtils.cpp
--- a/Engine/Plugins/Enterprise/DataprepEditor/Source/DataprepCore/Private/DataprepCoreUtils.cpp
+++ b/Engine/Plugins/Enterprise/DataprepEditor/Source/DataprepCore/Private/DataprepCoreUtils.cpp
@@ -20,6 +20,7 @@
 #include "SelectionSystem/DataprepFilter.h"
 #include "SelectionSystem/DataprepSelectionTransform.h"
 
+#include "AssetRegistryModule.h"
 #include "Engine/Level.h"
 #include "Engine/StaticMesh.h"
 #include "Engine/World.h"
@@ -38,6 +39,7 @@
 #include "Templates/SubclassOf.h"
 #include "UObject/StrongObjectPtr.h"
 #include "UObject/UObjectHash.h"
+#include "UObject/UObjectIterator.h"
 
 #if WITH_EDITOR
 #include "ActorEditorUtils.h"
@@ -368,12 +370,11 @@ bool FDataprepCoreUtils::ExecuteDataprep(UDataprepAssetInterface* DataprepAssetI
 
 			FDataprepCoreUtils::PurgeObjects( MoveTemp( ObjectsToDelete ) );
 
+			// Erase all temporary packages and files created by the Dataprep asset
+			DeleteTemporaryFolders( TransientContentFolder );
+
 			// Restore LogStaticMesh verbosity
 			LogStaticMesh.SetVerbosity( PrevLogStaticMeshVerbosity );
-
-			// Erase all temporary files created by the Dataprep asset
-			static FString AbsolutePath = FPaths::ConvertRelativePathToFull( DataprepCorePrivateUtils::GetRootTemporaryDir() / RelativeTempFolder );
-			IFileManager::Get().DeleteDirectory( *AbsolutePath, false, true );
 		}
 
 		return bSuccessfulExecute;
@@ -754,5 +755,104 @@ void FDataprepCoreUtils::GetActorsFromWorld(const UWorld* World, TArray<AActor*>
 	DataprepCoreUtilsPrivate::GetActorsFromWorld( World, OutActors );
 }
 
+void FDataprepCoreUtils::DeleteTemporaryFolders(const FString& BaseTemporaryPath)
+{
+	TRACE_CPUPROFILER_EVENT_SCOPE(FDataprepCoreUtils::DeleteTemporaryFolders);
+
+	FAssetRegistryModule& AssetRegistryModule = FModuleManager::Get().LoadModuleChecked<FAssetRegistryModule>( TEXT("AssetRegistry") );
+	IAssetRegistry& AssetRegistry = AssetRegistryModule.Get();
+
+	TArray< UObject* > ObjectsToDelete;
+
+	// Find all registered objects which are in  memory and under the temporary path
+	{
+		// Query for a list of assets in the path to delete
+		FARFilter Filter;
+		Filter.bRecursivePaths = true;
+		Filter.PackagePaths.Emplace( *BaseTemporaryPath );
+
+		TArray<FAssetData> AssetDataList;
+		AssetRegistry.GetAssets( Filter, AssetDataList );
+
+		ObjectsToDelete.Reserve( AssetDataList.Num() );
+		for(const FAssetData& AssetData : AssetDataList)
+		{
+			FSoftObjectPath ObjectPath( AssetData.ObjectPath.ToString() );
+
+			if(UObject* Object = ObjectPath.ResolveObject())
+			{
+				ObjectsToDelete.Add( Object );
+			}
+		}
+	}
+
+	// Find all packages under the temporary path
+	{
+		for (TObjectIterator<UPackage> It; It; ++It)
+		{
+			UPackage*	Package = *It;
+
+			if(Package->GetName().StartsWith( BaseTemporaryPath ))
+			{
+				// Remove package path from asset registry
+				AssetRegistry.RemovePath(Package->GetPathName());
+
+				ObjectsToDelete.Add( Package );
+			}
+		}
+	}
+
+	if(ObjectsToDelete.Num() > 0)
+	{
+		PurgeObjects( MoveTemp( ObjectsToDelete ) );
+	}
+	// No object to delete, force a garbage collection anyway
+	else
+	{
+		CollectGarbage( GARBAGE_COLLECTION_KEEPFLAGS );
+	}
+
+	// Delete all assets on disk
+	{
+		struct FEmptyFolderVisitor : public IPlatformFile::FDirectoryVisitor
+		{
+			bool bIsEmpty;
+
+			FEmptyFolderVisitor()
+				: bIsEmpty( true )
+			{
+			}
+
+			virtual bool Visit(const TCHAR* FilenameOrDirectory, bool bIsDirectory) override
+			{
+				if ( !bIsDirectory )
+				{
+					bIsEmpty = false;
+					return false; // abort searching
+				}
+
+				return true; // continue searching
+			}
+		};
+
+		IFileManager& FileManager = IFileManager::Get();
+
+		FString PathToDeleteOnDisk;
+		if ( FPackageName::TryConvertLongPackageNameToFilename( BaseTemporaryPath, PathToDeleteOnDisk ) )
+		{
+			if( FileManager.DirectoryExists( *PathToDeleteOnDisk ) )
+			{
+				// Look for files on disk in case the folder contains things not tracked by the asset registry
+				FEmptyFolderVisitor EmptyFolderVisitor;
+				FileManager.IterateDirectoryRecursively( *PathToDeleteOnDisk, EmptyFolderVisitor );
+
+				if ( EmptyFolderVisitor.bIsEmpty && IFileManager::Get().DeleteDirectory( *PathToDeleteOnDisk, false, true ) )
+				{
+					AssetRegistry.RemovePath( BaseTemporaryPath );
+				}
+			}
+		}
+	}
+}
 
 #undef LOCTEXT_NAMESPACE
\ No newline at end of file
diff --git a/Engine/Plugins/Enterprise/DataprepEditor/Source/DataprepCore/Public/DataprepCoreUtils.h b/Engine/Plugins/Enterprise/DataprepEditor/Source/DataprepCore/Public/DataprepCoreUtils.h
--- a/Engine/Plugins/Enterprise/DataprepEditor/Source/DataprepCore/Public/DataprepCoreUtils.h
+++ b/Engine/Plugins/Enterprise/DataprepEditor/Source/DataprepCore/Public/DataprepCoreUtils.h
@@ -26,14 +26,14 @@ class DATAPREPCORE_API FDataprepCoreUtils
 public:
 
 	/**
-	 * Return the dataprep asset that own the object, if the object is part of a dataprep asset
-	 * @return nullptr if the object is not a part of a dataprep asset
+	 * Return the Dataprep asset that own the object, if the object is part of a Dataprep asset
+	 * @return nullptr if the object is not a part of a Dataprep asset
 	 */
 	static UDataprepAsset* GetDataprepAssetOfObject(UObject* Object);
 
 	/**
-	 * Return the dataprep action asset that own the object, if the object is part of a dataprep action asset
-	 * @return nullptr if the object is not a part of a dataprep action asset
+	 * Return the Dataprep action asset that own the object, if the object is part of a Dataprep action asset
+	 * @return nullptr if the object is not a part of a Dataprep action asset
 	 */
 	static UDataprepActionAsset* GetDataprepActionAssetOf(UObject* Object);
 
@@ -79,15 +79,15 @@ public:
 	}
 
 	/**
-	 * Execute a dataprep asset (import, execute and commit)
-	 * @param DataprepAssetInterface The dataprep asset to execute
-	 * @param Logger A optional log manager for the dataprep producers, operations and consumer
-	 * @param Reporter A optional progress reporter for the dataprep producers, operations and consumer
+	 * Execute a Dataprep asset (import, execute and commit)
+	 * @param DataprepAssetInterface The Dataprep asset to execute
+	 * @param Logger A optional log manager for the Dataprep producers, operations and consumer
+	 * @param Reporter A optional progress reporter for the Dataprep producers, operations and consumer
 	 */
 	static bool ExecuteDataprep(UDataprepAssetInterface* DataprepAssetInterface, const TSharedPtr<IDataprepLogger>& Logger, const TSharedPtr<IDataprepProgressReporter>& Reporter);
 
 	/**
-	 * Check if class can be used to create a step for a dataprep action
+	 * Check if class can be used to create a step for a Dataprep action
 	 * @param StepType The class for the creation of the step
 	 * @param OutValidRootClass The valid root class if found.
 	 * @param OutMessageIfInvalid A message to report if the class is invalid.
@@ -236,5 +236,12 @@ public:
 	 */
 	static void GetActorsFromWorld(const UWorld* World, TArray<AActor*>& OutActors);
 
+	/**
+	 * Delete the specified folder including all assets and sub-folders inside it.
+	 * @param BaseTemporaryPath			Root folder to start from
+	 * @remark This is called at the end of the execution of a Dataprep asset to remove all
+	 *		   the temporary assets and folders created during its execution.
+	 */
+	static void DeleteTemporaryFolders( const FString& BaseTemporaryPath );
 };
 
diff --git a/Engine/Plugins/Enterprise/DataprepEditor/Source/DataprepEditor/Private/DataprepEditor.cpp b/Engine/Plugins/Enterprise/DataprepEditor/Source/DataprepEditor/Private/DataprepEditor.cpp
--- a/Engine/Plugins/Enterprise/DataprepEditor/Source/DataprepEditor/Private/DataprepEditor.cpp
+++ b/Engine/Plugins/Enterprise/DataprepEditor/Source/DataprepEditor/Private/DataprepEditor.cpp
@@ -110,29 +110,6 @@ private:
 	FString Text;
 };
 
-namespace DataprepEditorUtil
-{
-	void DeleteActor(AActor* Actor, UWorld* World)
-	{
-		if (Actor == nullptr || World != Actor->GetWorld())
-		{
-			return;
-		}
-
-		TArray<AActor*> Children;
-		Actor->GetAttachedActors(Children);
-
-		for (AActor* ChildActor : Children)
-		{
-			DeleteActor(ChildActor, World);
-		}
-
-		World->DestroyActor(Actor, false, true);
-	}
-
-	void DeleteTemporaryPackage( const FString& PathToDelete );
-}
-
 FDataprepEditor::FDataprepEditor()
 	: bWorldBuilt(false)
 	, bIsFirstRun(false)
@@ -574,7 +551,7 @@ void FDataprepEditor::ResetBuildWorld()
 	bWorldBuilt = false;
 	CleanPreviewWorld();
 	UpdatePreviewPanels();
-	DataprepEditorUtil::DeleteTemporaryPackage( GetTransientContentFolder() );
+	FDataprepCoreUtils::DeleteTemporaryFolders( GetTransientContentFolder() );
 }
 
 void FDataprepEditor::CleanPreviewWorld()
@@ -1265,91 +1242,4 @@ void FDataprepEditor::ClearPreviewedObjects()
 	SetPreviewedObjects( MakeArrayView<UDataprepParameterizableObject*>( nullptr, 0 ) );
 }
 
-void DataprepEditorUtil::DeleteTemporaryPackage( const FString& PathToDelete )
-{
-	// See ContentBrowserUtils::LoadAssetsIfNeeded
-	// See ContentBrowserUtils::DeleteFolders
-
-	FAssetRegistryModule& AssetRegistryModule = FModuleManager::Get().LoadModuleChecked<FAssetRegistryModule>(TEXT("AssetRegistry"));
-
-	// Form a filter from the path to delete
-	FARFilter Filter;
-	Filter.bRecursivePaths = true;
-	new (Filter.PackagePaths) FName( *PathToDelete );
-
-	// Query for a list of assets in the path to delete
-	TArray<FAssetData> AssetDataList;
-	AssetRegistryModule.Get().GetAssets(Filter, AssetDataList);
-
-	// Delete all registered objects which are in  memory and under this path
-	{
-		TArray<UObject*> AssetsToDelete;
-		AssetsToDelete.Reserve(AssetDataList.Num());
-		for(const FAssetData& AssetData : AssetDataList)
-		{
-			FSoftObjectPath ObjectPath( AssetData.ObjectPath.ToString() );
-
-			if(UObject* Object = ObjectPath.ResolveObject())
-			{
-				AssetsToDelete.Add(Object);
-			}
-		}
-
-		if(AssetsToDelete.Num() > 0)
-		{
-			ObjectTools::DeleteObjects(AssetsToDelete, false);
-		}
-	}
-
-	// Delete all assets not in memory but on disk
-	{
-		struct FEmptyFolderVisitor : public IPlatformFile::FDirectoryVisitor
-		{
-			bool bIsEmpty;
-
-			FEmptyFolderVisitor()
-				: bIsEmpty(true)
-			{
-			}
-
-			virtual bool Visit(const TCHAR* FilenameOrDirectory, bool bIsDirectory) override
-			{
-				if (!bIsDirectory)
-				{
-					bIsEmpty = false;
-					return false; // abort searching
-				}
-
-				return true; // continue searching
-			}
-		};
-
-		FString PathToDeleteOnDisk;
-		if (FPackageName::TryConvertLongPackageNameToFilename(PathToDelete, PathToDeleteOnDisk))
-		{
-			if(IFileManager::Get().DirectoryExists( *PathToDeleteOnDisk ))
-			{
-				// Look for files on disk in case the folder contains things not tracked by the asset registry
-				FEmptyFolderVisitor EmptyFolderVisitor;
-				IFileManager::Get().IterateDirectoryRecursively(*PathToDeleteOnDisk, EmptyFolderVisitor);
-
-				if (EmptyFolderVisitor.bIsEmpty && IFileManager::Get().DeleteDirectory(*PathToDeleteOnDisk, false, true))
-				{
-					AssetRegistryModule.Get().RemovePath(PathToDelete);
-				}
-			}
-			// Request deletion anyway
-			else
-			{
-				AssetRegistryModule.Get().RemovePath(PathToDelete);
-			}
-		}
-	}
-
-	// Check that no asset remains
-	AssetDataList.Reset();
-	AssetRegistryModule.Get().GetAssets(Filter, AssetDataList);
-	//ensure(AssetDataList.Num() == 0);
-}
-
 #undef LOCTEXT_NAMESPACE
diff --git a/Engine/Plugins/Enterprise/DatasmithImporter/Source/DatasmithImporter/Private/DatasmithFileProducer.cpp b/Engine/Plugins/Enterprise/DatasmithImporter/Source/DatasmithImporter/Private/DatasmithFileProducer.cpp
--- a/Engine/Plugins/Enterprise/DatasmithImporter/Source/DatasmithImporter/Private/DatasmithFileProducer.cpp
+++ b/Engine/Plugins/Enterprise/DatasmithImporter/Source/DatasmithImporter/Private/DatasmithFileProducer.cpp
@@ -455,10 +455,6 @@ void UDatasmithFileProducer::PreventNameCollision()
 		{
 			if( UObject* Object = Assets[Index].Get() )
 			{
-				// Ensure object's package is transient and not public
-				Object->GetOutermost()->ClearFlags( RF_Public );
-				Object->GetOutermost()->SetFlags( RF_Transient );
-
 				PathsToDelete.Add( Object->GetOutermost()->GetPathName() );
 
 				if( Cast<UStaticMesh>( Object ) != nullptr )
