Fixed crash when Dataprep assets prior to 4.25 are saved then reloaded

#rnx

diff --git a/Engine/Plugins/Enterprise/DataprepEditor/Source/DataprepCore/Private/DataprepAsset.cpp b/Engine/Plugins/Enterprise/DataprepEditor/Source/DataprepCore/Private/DataprepAsset.cpp
--- a/Engine/Plugins/Enterprise/DataprepEditor/Source/DataprepCore/Private/DataprepAsset.cpp
+++ b/Engine/Plugins/Enterprise/DataprepEditor/Source/DataprepCore/Private/DataprepAsset.cpp
@@ -68,6 +68,11 @@ void UDataprepAsset::PostLoad()
 			DataprepRecipeBP_DEPRECATED = nullptr;
 			StartNode_DEPRECATED = nullptr;
 			bMarkDirty = true;
+
+			// Actions must change owner, but it cannot be done in the PostLoad.
+			// Register to be called back when FCoreUObjectDelegates::OnAssetLoad broadcasts
+			// so the change of ownership can be done
+			FCoreUObjectDelegates::OnAssetLoaded.AddUObject(this, &UDataprepAsset::OnOldAssetLoaded);
 		}
 
 		if ( !Parameterization )
@@ -599,4 +604,22 @@ bool UDataprepAsset::RemoveActions(const TArray<int32>& Indices)
 	return false;
 }
 
+void UDataprepAsset::OnOldAssetLoaded(UObject* Asset)
+{
+	if(Asset == this)
+	{
+		// Move ownership from the deprecated K2Node to the Dataprep asset itself
+		for(UDataprepActionAsset* ActionAsset : ActionAssets)
+		{
+			if(ActionAsset && ActionAsset->GetOuter() != this)
+			{
+				ActionAsset->Rename(nullptr, this, REN_NonTransactional | REN_DontCreateRedirectors);
+			}
+		}
+
+		// Unregister to the OnAssetLoad event as it is not needed anymore
+		FCoreUObjectDelegates::OnAssetLoaded.RemoveAll(this);
+	}
+}
+
 #undef LOCTEXT_NAMESPACE
diff --git a/Engine/Plugins/Enterprise/DataprepEditor/Source/DataprepCore/Public/DataprepAsset.h b/Engine/Plugins/Enterprise/DataprepEditor/Source/DataprepCore/Public/DataprepAsset.h
--- a/Engine/Plugins/Enterprise/DataprepEditor/Source/DataprepCore/Public/DataprepAsset.h
+++ b/Engine/Plugins/Enterprise/DataprepEditor/Source/DataprepCore/Public/DataprepAsset.h
@@ -224,6 +224,10 @@ protected:
 	UDataprepContentConsumer* Consumer_DEPRECATED;
 #endif
 
+private:
+	/** Handler for when an Dataprep asset from version prior to 4.25 is loaded */
+	void OnOldAssetLoaded(UObject* Asset);
+
 private:
 	/** DEPRECATED: Pointer to the entry node of the pipeline blueprint previously used to process input data */
 	UPROPERTY(meta = (DeprecatedProperty, DeprecationMessage = "Using directly ActionAssets property instead of Blueprint to manage actions."))
