Fixed crash when undoing the insertion of an action

#rnx

diff --git a/Engine/Plugins/Enterprise/DataprepEditor/Source/DataprepCore/Private/DataprepAsset.cpp b/Engine/Plugins/Enterprise/DataprepEditor/Source/DataprepCore/Private/DataprepAsset.cpp
--- a/Engine/Plugins/Enterprise/DataprepEditor/Source/DataprepCore/Private/DataprepAsset.cpp
+++ b/Engine/Plugins/Enterprise/DataprepEditor/Source/DataprepCore/Private/DataprepAsset.cpp
@@ -419,6 +419,8 @@ bool UDataprepAsset::InsertAction(int32 Index)
 {
 	if ( Index >= 0 && Index <= ActionAssets.Num() )
 	{
+		Modify();
+
 		UDataprepActionAsset* Action = NewObject<UDataprepActionAsset>( this, UDataprepActionAsset::StaticClass(), NAME_None, RF_Transactional );
 		Action->SetLabel( TEXT("New Action") );
 
