Fixed crash with CreateProxyMesh operation in Dataprep editor
  - Root cause: Missing call to PostEdit on materials after import

#rnx

diff --git a/Engine/Plugins/Enterprise/DataprepEditor/Source/DataprepCore/Private/DataprepCoreUtils.cpp b/Engine/Plugins/Enterprise/DataprepEditor/Source/DataprepCore/Private/DataprepCoreUtils.cpp
--- a/Engine/Plugins/Enterprise/DataprepEditor/Source/DataprepCore/Private/DataprepCoreUtils.cpp
+++ b/Engine/Plugins/Enterprise/DataprepEditor/Source/DataprepCore/Private/DataprepCoreUtils.cpp
@@ -661,20 +661,17 @@ void FDataprepCoreUtils::BuildAssets(const TArray<TWeakObjectPtr<UObject>>& Asse
 	// Force compilation of materials which have no render proxy
 	if(MaterialInterfaces.Num() > 0)
 	{
-		auto MustCompile = [](UMaterialInterface* MaterialInterface) -> bool
-		{
-			if(MaterialInterface == nullptr)
-			{
-				return false;
-			}
+		// Post-processing on materials inspired from FDatasmithImporterImpl::CompileMaterial
 
+		// Check if material must be recompiled
+		auto MustRecompile = [](UMaterialInterface* MaterialInterface) -> bool
+		{
 			// Force recompilation of constant material instances which either override blend mode or any static switch
 			if(UMaterialInstanceConstant* ConstantMaterialInstance = Cast< UMaterialInstanceConstant >(MaterialInterface))
 			{
 				// If BlendMode override property has been changed, make sure this combination of the parent material is compiled
 				if ( ConstantMaterialInstance->BasePropertyOverrides.bOverride_BlendMode == true )
 				{
-					ConstantMaterialInstance->ForceRecompileForRendering();
 					return true;
 				}
 				else
@@ -687,7 +684,6 @@ void FDataprepCoreUtils::BuildAssets(const TArray<TWeakObjectPtr<UObject>>& Asse
 					{
 						if ( Switch.bOverride )
 						{
-							ConstantMaterialInstance->ForceRecompileForRendering();
 							return true;
 						}
 					}
@@ -708,10 +704,16 @@ void FDataprepCoreUtils::BuildAssets(const TArray<TWeakObjectPtr<UObject>>& Asse
 			++AssetBuiltCount;
 			Task.ReportNextStep(FText::Format(LOCTEXT( "BuildAssets_Building_Materials", "Building materials ({0} / {1})" ), AssetBuiltCount, AssetToBuildCount), 1.0f);
 
-			if( MustCompile( MaterialInterface ) )
+			if(MaterialInterface)
 			{
-				FPropertyChangedEvent EmptyPropertyUpdateStruct( nullptr );
-				MaterialInterface->PostEditChangeProperty( EmptyPropertyUpdateStruct );
+				if( MustRecompile( MaterialInterface ) )
+				{
+					MaterialInterface->ForceRecompileForRendering();
+				}
+
+				// Do not assume the material has been properly initialized, force post-edit on it
+				MaterialInterface->PreEditChange( nullptr );
+				MaterialInterface->PostEditChange();
 			}
 		}
 	}
