Deprecated FMeshSubstitutionDataTable UStruct

#rnx

diff --git a/Engine/Plugins/Enterprise/DataprepEditor/Source/DataprepLibraries/Private/DataprepOperationsLibrary.cpp b/Engine/Plugins/Enterprise/DataprepEditor/Source/DataprepLibraries/Private/DataprepOperationsLibrary.cpp
--- a/Engine/Plugins/Enterprise/DataprepEditor/Source/DataprepLibraries/Private/DataprepOperationsLibrary.cpp
+++ b/Engine/Plugins/Enterprise/DataprepEditor/Source/DataprepLibraries/Private/DataprepOperationsLibrary.cpp
@@ -630,24 +630,8 @@ void UDataprepOperationsLibrary::SubstituteMesh(const TArray<UObject*>& Selected
 	SubstituteMesh( SelectedObjects, MeshSearch, StringMatch, MeshesUsed, MeshSubstitute );
 }
 
-void UDataprepOperationsLibrary::SubstituteMeshesByTable(const TArray<UObject*>& SelectedObjects, const UDataTable * DataTable)
+void UDataprepOperationsLibrary::SubstituteMeshesByTable(const TArray<UObject*>& , const UDataTable* )
 {
-	if (DataTable == nullptr || DataTable->GetRowStruct() == nullptr || !DataTable->GetRowStruct()->IsChildOf(FMeshSubstitutionDataTable::StaticStruct()))
-	{
-		return;
-	}
-
-	TArray<UStaticMesh*> MeshesUsed = DataprepOperationsLibraryUtil::GetUsedMeshes(SelectedObjects);
-
-	const TMap<FName, uint8*>&  MeshTableRowMap = DataTable->GetRowMap();
-	for (auto& MeshTableRowEntry : MeshTableRowMap)
-	{
-		const FMeshSubstitutionDataTable* MeshRow = (const FMeshSubstitutionDataTable*)MeshTableRowEntry.Value;
-		if (MeshRow != nullptr && MeshRow->MeshReplacement != nullptr)
-		{
-			SubstituteMesh( SelectedObjects, MeshRow->SearchString, MeshRow->StringMatch, MeshesUsed, MeshRow->MeshReplacement );
-		}
-	}
 }
 
 void UDataprepOperationsLibrary::SubstituteMesh(const TArray<UObject*>& SelectedObjects, const FString& MeshSearch, EEditorScriptingStringMatchType StringMatch, const TArray<UStaticMesh*>& MeshList, UStaticMesh* MeshSubstitute)
diff --git a/Engine/Plugins/Enterprise/DataprepEditor/Source/DataprepLibraries/Private/DataprepOperationsLibrary.h b/Engine/Plugins/Enterprise/DataprepEditor/Source/DataprepLibraries/Private/DataprepOperationsLibrary.h
--- a/Engine/Plugins/Enterprise/DataprepEditor/Source/DataprepLibraries/Private/DataprepOperationsLibrary.h
+++ b/Engine/Plugins/Enterprise/DataprepEditor/Source/DataprepLibraries/Private/DataprepOperationsLibrary.h
@@ -68,24 +68,24 @@ struct FMaterialSubstitutionDataTable : public FTableRowBase
 };
 
 /*
-* Simple struct for the table row used for UDataprepOperationsLibrary::SubstituteMaterials
+* DEPRECATED - Simple struct for the table row used for UDataprepOperationsLibrary::SubstituteMaterials
 */
-USTRUCT(BlueprintType)
-struct FMeshSubstitutionDataTable : public FTableRowBase
+USTRUCT(meta=(Deprecated = "4.25.0"))
+struct FMeshSubstitutionDataTable
 {
 	GENERATED_BODY()
 
-	/** Name of the mesh(es) to search for. Wildcard is supported */
-	UPROPERTY(EditAnywhere, BlueprintReadOnly, Category = "MeshSubstitutionTable")
-	FString SearchString;
+	/** DEPRECATED - Name of the mesh(es) to search for. Wildcard is supported */
+	UPROPERTY()
+	FString SearchString_DEPRECATED;
 
-	/** Type of matching to perform with SearchString string */
-	UPROPERTY(EditAnywhere, BlueprintReadOnly, Category = "MeshSubstitutionTable")
-	EEditorScriptingStringMatchType StringMatch;
+	/** DEPRECATED - Type of matching to perform with SearchString string */
+	UPROPERTY()
+	EEditorScriptingStringMatchType StringMatch_DEPRECATED;
 
-	/** Mesh to use for the substitution */
-	UPROPERTY(EditAnywhere, BlueprintReadOnly, Category = "MeshSubstitutionTable")
-	UStaticMesh* MeshReplacement;
+	/** DEPRECATED - Mesh to use for the substitution */
+	UPROPERTY()
+	UStaticMesh* MeshReplacement_DEPRECATED;
 };
 
 /*
@@ -276,7 +276,7 @@ public:
 	 * @param DataTable:		Data table to use for the substitution
 	 * @remark: SubstituteMesh is called for each entry of the input data table
 	 */
-	UFUNCTION(BlueprintCallable, Category = "Dataprep | Operation")
+	UFUNCTION(meta=(Deprecated = "4.25.0"))
 	static void SubstituteMeshesByTable(const TArray<UObject*>& SelectedObjects, const UDataTable* DataTable);
 
 	/**
