Improved fluidity of mouse's and node's movementsduring the drag of an action node

#rnx

diff --git a/Engine/Plugins/Enterprise/DataprepEditor/Source/DataprepEditor/Private/DataprepEditorStyle.cpp b/Engine/Plugins/Enterprise/DataprepEditor/Source/DataprepEditor/Private/DataprepEditorStyle.cpp
--- a/Engine/Plugins/Enterprise/DataprepEditor/Source/DataprepEditor/Private/DataprepEditorStyle.cpp
+++ b/Engine/Plugins/Enterprise/DataprepEditor/Source/DataprepEditor/Private/DataprepEditorStyle.cpp
@@ -30,8 +30,9 @@ void FDataprepEditorStyle::Initialize()
 	StyleSet->SetContentRoot(FPaths::EngineContentDir() / TEXT("Editor/Slate"));
 	StyleSet->SetCoreContentRoot(FPaths::EngineContentDir() / TEXT("Slate"));
 
-	const FVector2D Icon20x16(16.0f, 16.0f);
+	const FVector2D Icon16x16(16.0f, 16.0f);
 	const FVector2D Icon20x20(20.0f, 20.0f);
+	const FVector2D Icon24x24(24.0f, 24.0f);
 	const FVector2D Icon32x32(32.0f, 32.0f);
 	const FVector2D Icon40x40(40.0f, 40.0f);
 
@@ -72,6 +73,9 @@ void FDataprepEditorStyle::Initialize()
 		StyleSet->Set("Dataprep.Background.Black", FLinearColor(FColor(26, 26, 26)) );
 		StyleSet->Set("Dataprep.TextSeparator.Color", FLinearColor(FColor(200, 200, 200, 200)) );
 
+		StyleSet->Set("DataprepEditor.SoftwareCursor_Grab", new IMAGE_PLUGIN_BRUSH("Icons/cursor_grab", Icon20x20));
+		StyleSet->Set("DataprepEditor.SoftwareCursor_Hand", new IMAGE_PLUGIN_BRUSH("Icons/cursor_hand", Icon20x20));
+
 		StyleSet->Set("DataprepAction.Outter.Regular.Padding", FMargin(1.f, 2.f, 4.f, 2.f));
 		StyleSet->Set("DataprepAction.Outter.Selected.Padding", FMargin(0.f));
 		StyleSet->Set("DataprepAction.Body.Padding", FMargin(6.f, 2.f, 2.f, 2.f));
diff --git a/Engine/Plugins/Enterprise/DataprepEditor/Source/DataprepEditor/Private/Widgets/DataprepGraph/DragDropActionNodeImpl.cpp b/Engine/Plugins/Enterprise/DataprepEditor/Source/DataprepEditor/Private/Widgets/DataprepGraph/DragDropActionNodeImpl.cpp
--- a/Engine/Plugins/Enterprise/DataprepEditor/Source/DataprepEditor/Private/Widgets/DataprepGraph/DragDropActionNodeImpl.cpp
+++ b/Engine/Plugins/Enterprise/DataprepEditor/Source/DataprepEditor/Private/Widgets/DataprepGraph/DragDropActionNodeImpl.cpp
@@ -31,6 +31,14 @@ class FDragDropActionNodeImpl : public FDragDropActionNode
 public:
 	DRAG_DROP_OPERATOR_TYPE(FDragDropActionNodeImpl, FDragDropActionNode)
 
+	virtual void Construct() override
+	{
+		// Hide cursor when starting the drag. The track node will restore it
+		MouseCursor = EMouseCursor::None;
+		bCreateNewWindow = false;
+		FDragDropActionNode::Construct();
+	}
+
 	// FDragDropOperation interface
 	virtual void OnDrop( bool bDropWasHandled, const FPointerEvent& MouseEvent ) override;
 	virtual void OnDragged( const class FDragDropEvent& DragDropEvent ) override;
@@ -48,73 +56,33 @@ public:
 
 TSharedRef<FDragDropActionNode> FDragDropActionNode::New(const TSharedRef<SDataprepGraphTrackNode>& InTrackNodePtr, const TSharedRef<SDataprepGraphActionNode>& InDraggedNode)
 {
-	FDragDropActionNodeImpl* OperationImpl = new FDragDropActionNodeImpl;
+	TSharedPtr<FDragDropActionNodeImpl> OperationImpl = MakeShared<FDragDropActionNodeImpl>();
 
 	OperationImpl->TrackNodePtr = InTrackNodePtr;
 	OperationImpl->ActionNodePtr = InDraggedNode;
 
-	OperationImpl->bCreateNewWindow = false;
 	OperationImpl->Construct();
 
 	InTrackNodePtr->OnStartNodeDrag(InDraggedNode);
 
 	TSharedRef<FDragDropActionNode> Operation = MakeShareable(new FDragDropActionNode);
-	Operation->Impl = TSharedPtr<FDragDropActionNode>(OperationImpl);
+	Operation->Impl = StaticCastSharedPtr<FDragDropActionNode>(OperationImpl);
 
 	return Operation;
 }
 
 void FDragDropActionNodeImpl::OnDrop(bool bDropWasHandled, const FPointerEvent& MouseEvent)
 {
-	int32 NewExecutionOrder = TrackNodePtr->OnEndNodeDrag();
-
-	if(bDropWasHandled)
-	{
-		if(UDataprepAsset* DataprepAsset = TrackNodePtr->GetDataprepAsset())
-		{
-			FModifierKeysState ModifierKeyState = FSlateApplication::Get().GetModifierKeys();
-			const bool bCopyRequested = ModifierKeyState.IsControlDown() || ModifierKeyState.IsCommandDown();
-
-			FScopedTransaction Transaction( bCopyRequested ? LOCTEXT("OnDropCopy", "Add/Insert action") : LOCTEXT("OnDropMove", "Move action") );
-			bool bTransactionSuccessful = false;
-
-			if(bCopyRequested)
-			{
-				if(NewExecutionOrder >= DataprepAsset->GetActionCount())
-				{
-					bTransactionSuccessful = DataprepAsset->AddAction(ActionNodePtr->GetDataprepAction()) != INDEX_NONE;
-				}
-				else
-				{
-					bTransactionSuccessful = DataprepAsset->InsertAction(ActionNodePtr->GetDataprepAction(), NewExecutionOrder);
-				}
-			}
-			else if( NewExecutionOrder != ActionNodePtr->GetExecutionOrder())
-			{
-				bTransactionSuccessful = DataprepAsset->MoveAction(ActionNodePtr->GetExecutionOrder(), NewExecutionOrder);
-			}
-			else
-			{
-				TrackNodePtr->RefreshLayout();
-			}
-
-			if(!bTransactionSuccessful)
-			{
-				Transaction.Cancel();
-			}
-		}
-	}
-	else
-	{
-		TrackNodePtr->RefreshLayout();
-	}
+	TrackNodePtr->OnEndNodeDrag(bDropWasHandled);
 
 	FDragDropOperation::OnDrop(bDropWasHandled, MouseEvent);
 }
 
 void FDragDropActionNodeImpl::OnDragged(const FDragDropEvent& DragDropEvent)
 {
-	TrackNodePtr->OnNodeDragged( ActionNodePtr, DragDropEvent.GetScreenSpacePosition(), DragDropEvent.GetCursorDelta() );
+	const bool bNodeDragged = TrackNodePtr->OnNodeDragged( ActionNodePtr, DragDropEvent.GetScreenSpacePosition(), DragDropEvent.GetCursorDelta() );
+
+	MouseCursor = bNodeDragged ? EMouseCursor::None : EMouseCursor::SlashedCircle;
 
 	FDragDropOperation::OnDragged(DragDropEvent);
 }
diff --git a/Engine/Plugins/Enterprise/DataprepEditor/Source/DataprepEditor/Private/Widgets/DataprepGraph/SDataprepGraphActionNode.cpp b/Engine/Plugins/Enterprise/DataprepEditor/Source/DataprepEditor/Private/Widgets/DataprepGraph/SDataprepGraphActionNode.cpp
--- a/Engine/Plugins/Enterprise/DataprepEditor/Source/DataprepEditor/Private/Widgets/DataprepGraph/SDataprepGraphActionNode.cpp
+++ b/Engine/Plugins/Enterprise/DataprepEditor/Source/DataprepEditor/Private/Widgets/DataprepGraph/SDataprepGraphActionNode.cpp
@@ -368,7 +368,7 @@ void SDataprepGraphActionNode::Construct(const FArguments& InArgs, UDataprepGrap
 	GraphNode = InActionNode;
 	DataprepEditor = InArgs._DataprepEditor;
 
-	SetCursor(EMouseCursor::ResizeLeftRight);
+	SetCursor(EMouseCursor::CardinalCross);
 	UpdateGraphNode();
 
 	ProxyNodePtr = SNew(SDataprepGraphActionProxyNode, SharedThis(this))
@@ -691,7 +691,7 @@ void SDataprepGraphActionNode::OnDragEnter(const FGeometry& MyGeometry, const FD
 FSlateColor SDataprepGraphActionNode::GetInsertColor(int32 Index)
 {
 	static const FSlateColor BackgroundColor = FDataprepEditorStyle::GetColor("DataprepActionStep.BackgroundColor");
-	static const FSlateColor DragAndDrop = FDataprepEditorStyle::GetColor("DataprepActionStep.Separator.Color")/*FDataprepEditorStyle::GetColor("DataprepActionStep.DragAndDrop")*/;
+	static const FSlateColor DragAndDrop = FDataprepEditorStyle::GetColor("DataprepActionStep.Separator.Color");
 
 
 	return Index == InsertIndex ? DragAndDrop : FLinearColor::Transparent;
diff --git a/Engine/Plugins/Enterprise/DataprepEditor/Source/DataprepEditor/Private/Widgets/DataprepGraph/SDataprepGraphTrackNode.cpp b/Engine/Plugins/Enterprise/DataprepEditor/Source/DataprepEditor/Private/Widgets/DataprepGraph/SDataprepGraphTrackNode.cpp
--- a/Engine/Plugins/Enterprise/DataprepEditor/Source/DataprepEditor/Private/Widgets/DataprepGraph/SDataprepGraphTrackNode.cpp
+++ b/Engine/Plugins/Enterprise/DataprepEditor/Source/DataprepEditor/Private/Widgets/DataprepGraph/SDataprepGraphTrackNode.cpp
@@ -9,15 +9,17 @@
 #include "DataprepGraph/DataprepGraphActionNode.h"
 #include "SchemaActions/DataprepDragDropOp.h"
 
-#include "Layout/Children.h"
 #include "Widgets/DataprepGraph/SDataprepGraphActionNode.h"
 #include "Widgets/DataprepGraph/SDataprepGraphEditor.h"
 #include "Widgets/DataprepWidgets.h"
 
 #include "Framework/Application/SlateApplication.h"
 #include "Framework/Application/SlateUser.h"
-#include "SGraphPanel.h"
+#include "GenericPlatform/ICursor.h"
+#include "Layout/Children.h"
 #include "NodeFactory.h"
+#include "ScopedTransaction.h"
+#include "SGraphPanel.h"
 #include "Widgets/Colors/SColorBlock.h"
 #include "Widgets/Images/SImage.h"
 #include "Widgets/Layout/SBox.h"
@@ -25,10 +27,16 @@
 #include "Widgets/SBoxPanel.h"
 #include "Widgets/SOverlay.h"
 
+#if PLATFORM_WINDOWS || PLATFORM_HOLOLENS
+#include "Windows/WindowsHWrapper.h"
+#endif
+
 #define LOCTEXT_NAMESPACE "DataprepGraphEditor"
 
 FVector2D SDataprepGraphTrackNode::TrackAnchor( 15.f, 40.f );
 
+typedef TFunctionRef<float(const FVector2D&, float)> DragCallback;
+
 class SDataprepEmptyActionNode : public SVerticalBox
 {
 public:
@@ -102,10 +110,43 @@ public:
 	SLATE_BEGIN_ARGS(SDataprepGraphTrackWidget) {}
 	SLATE_END_ARGS();
 
-
 	void Construct(const FArguments& InArgs, TSharedPtr<SDataprepGraphTrackNode> InTrackNode);
 
-	void SetActionSlot(int32 SlotIndex, const TSharedRef<SWidget>& Widget);
+	// SWidget interface
+	virtual int32 OnPaint(const FPaintArgs& Args, const FGeometry& AllottedGeometry, const FSlateRect& MyCullingRect, FSlateWindowElementList& OutDrawElements, int32 LayerId, const FWidgetStyle& InWidgetStyle, bool bParentEnabled) const override
+	{
+		int32 MaxLayerId = SConstraintCanvas::OnPaint(Args, AllottedGeometry, MyCullingRect, OutDrawElements, LayerId, InWidgetStyle, bParentEnabled);
+
+		// If dragging an action node, fake the cursor as it is hidden
+		if(IndexDragged != INDEX_NONE)
+		{
+			static const FSlateBrush* GrabBrush = FDataprepEditorStyle::GetBrush(TEXT("DataprepEditor.SoftwareCursor_Grab"));
+			static const FSlateBrush* HandBrush = FDataprepEditorStyle::GetBrush(TEXT("DataprepEditor.SoftwareCursor_Hand"));
+
+			++MaxLayerId;
+			const FSlateBrush* Brush = bMouseInScope ? GrabBrush : HandBrush;
+
+			FSlateDrawElement::MakeBox(
+				OutDrawElements,
+				MaxLayerId,
+				AllottedGeometry.ToPaintGeometry( GetSoftwareCursorPosition() - ( Brush->ImageSize / 2 ) , Brush->ImageSize / AllottedGeometry.Scale ),
+				Brush);
+#ifdef DEBUG_DRAG
+			FVector2D LocalCursorPosition = GetTickSpaceGeometry().AbsoluteToLocal(FSlateApplication::Get().GetCursorPos());
+			LocalCursorPosition.Y = TrackCursorOffset.Y - 30.f;
+
+			FSlateDrawElement::MakeBox(
+				OutDrawElements,
+				MaxLayerId,
+				AllottedGeometry.ToPaintGeometry( LocalCursorPosition - ( Brush->ImageSize / 2 ) , Brush->ImageSize / AllottedGeometry.Scale ),
+				Brush);
+#endif
+
+		}
+
+		return MaxLayerId;
+	}
+	// End of SWidget interface
 
 	const FVector2D& GetWorkingSize() const
 	{
@@ -132,38 +173,37 @@ public:
 		return LineSize;
 	}
 
-	FVector2D GetActioNodeAbscissaRange();
-
 	int32 GetHoveredActionNode(float Left, float Width);
 
-	float GetActionNodeAbscissa(int32 Index)
-	{
-		return ActionSlots.IsValidIndex(Index) ? ActionSlots[Index]->OffsetAttr.Get(FMargin()).Left : 0.f;
-	}
-
 	FVector2D GetNodeSize(const TSharedRef<SWidget>& Widget) const;
 
 	void RefreshLayout();
 
 	void OnStartNodeDrag(int32 Index);
 
-	float OnNodeDragged(float Delta);
+	void OnNodeDragged(float Delta, DragCallback ComputePan);
 
-	int32 OnEndNodeDrag();
+	void OnEndNodeDrag(FVector2D& OutSpacePostion, int32& OutDraggedIndex, int32&OutDroppedIndex);
 
-	bool IsDraggedNodeOnEnds();
+	FVector2D GetSoftwareCursorPosition() const
+	{
+		return FVector2D( DragSlot->OffsetAttr.Get(FMargin()).Left + TrackCursorOffset.X, TrackCursorOffset.Y);
+	}
+
+	FVector2D GetAbsoluteSoftwareCursorPosition() const
+	{
+		return GetTickSpaceGeometry().LocalToAbsolute(GetSoftwareCursorPosition());
+	}
 
 	/**
 	*/
-	void UpdateLayoutOnDrag(int8 DragDirection = 0);
+	void UpdateLayoutOnDrag();
 
 	/** Toggles display of action nodes between copy vs move drop modes */
-	float OnControlKeyDown(bool bKeyDown);
+	void OnControlKeyDown(bool bKeyDown, FVector2D& OutScreenSpacePosition);
 
 	void UpdateDragIndicator(FVector2D MousePosition);
 
-	int32 GetDragIndicatorIndex() { return DragIndicatorIndex; }
-
 	void ResetDragIndicator()
 	{
 		DragIndicatorIndex = INDEX_NONE;
@@ -173,7 +213,6 @@ public:
 		TrackSlots[SlotCount]->AttachWidget(SNullWidget::NullWidget);
 
 		UpdateLine(SlotCount-1);
-
 	}
 
 	FVector2D UpdateActionSlot(FMargin& SlotOffset, int32 SlotIndex, const TSharedRef<SWidget>& ActionNode, bool bIsEmpty = false)
@@ -192,7 +231,7 @@ public:
 		ActionSlot->AttachWidget(bIsEmpty ? SNew(SColorBlock).Color(FLinearColor::Transparent).Size( Size ) : ActionNode);
 
 		TrackSlots[SlotIndex]->OffsetAttr.Set( FMargin(SlotOffset.Left + Size.X * 0.5f - 16.f, TrackSlotTopPadding, 32, 32));
-		TrackSlots[SlotIndex]->AttachWidget(SlotIndex == HoveredDragIndex ? TrackSlotDrop.ToSharedRef() : TrackSlotRegular.ToSharedRef());
+		TrackSlots[SlotIndex]->AttachWidget(SlotIndex == IndexHovered ? TrackSlotDrop.ToSharedRef() : TrackSlotRegular.ToSharedRef());
 
 		SlotOffset.Left += Size.X;
 
@@ -203,8 +242,6 @@ public:
 
 	void CreateHelperWidgets();
 
-	void GetDraggedNodePositions(FVector2D& Left, FVector2D& Right, const FGeometry& TargetGeometry);
-
 	void UpdateLine(int32 LastIndex);
 
 	TWeakPtr<SDataprepGraphTrackNode> TrackNodePtr;
@@ -232,14 +269,15 @@ public:
 	int32 SlotCount;
 
 	SConstraintCanvas::FSlot* DragSlot;
-	int32 DragStartIndex;
-	int32 DragIndex;
-	int32 HoveredDragIndex;
+	int32 IndexDragged;
+	int32 IndexHovered;
 	int32 DragIndicatorIndex;
 	FVector2D LastDragPosition;
 	FVector2D AbscissaRange;
 	FVector2D TrackOffset;
+	FVector2D TrackCursorOffset;
 	bool bIsCopying;
+	bool bMouseInScope;
 
 	static float TrackDesiredHeight;
 	static float InterSlotSpacing;
@@ -261,6 +299,13 @@ void SDataprepGraphTrackNode::Construct(const FArguments& InArgs, UDataprepGraph
 	NodeFactory = InArgs._NodeFactory;
 
 	bNodeDragging = false;
+	bSkipNextDragUpdate = false;
+	bSkipRefreshLayout = false;
+	bCursorLeftOnLeft = false;
+	bCursorLeftOnRight = false;
+	bNodeDragJustStarted = false;
+	LastDragDirection = 0.f;
+
 	SetCursor(EMouseCursor::Default);
 	GraphNode = InNode;
 	check(GraphNode);
@@ -374,11 +419,6 @@ void SDataprepGraphTrackNode::UpdateGraphNode()
 	];
 }
 
-bool SDataprepGraphTrackNode::CustomPrepass(float LayoutScaleMultiplier)
-{
-	return false;
-}
-
 void SDataprepGraphTrackNode::OnActionsOrderChanged()
 {
 	if(!DataprepAssetPtr.IsValid())
@@ -430,7 +470,12 @@ bool SDataprepGraphTrackNode::RefreshLayout()
 	{
 		ensure(TrackWidgetPtr.IsValid());
 
-		TrackWidgetPtr->RefreshLayout();
+		if(!bSkipRefreshLayout && !bSkipNextDragUpdate)
+		{
+			TrackWidgetPtr->RefreshLayout();
+		}
+
+		bSkipRefreshLayout = false;
 
 		return true;
 	}
@@ -462,7 +507,7 @@ FReply SDataprepGraphTrackNode::OnDrop(const FGeometry& MyGeometry, const FDragD
 	TSharedPtr<FDataprepDragDropOp> DragActionNodeOp = DragDropEvent.GetOperationAs<FDataprepDragDropOp>();
 	if(DragActionNodeOp.IsValid())
 	{
-		int32 InsertIndex = TrackWidgetPtr->GetDragIndicatorIndex();
+		int32 InsertIndex = TrackWidgetPtr->DragIndicatorIndex;
 		TrackWidgetPtr->ResetDragIndicator();
 
 		DragActionNodeOp->DoDropOnTrack(DataprepAssetPtr.Get(), InsertIndex);
@@ -575,12 +620,26 @@ void SDataprepGraphTrackNode::OnStartNodeDrag(const TSharedRef<SDataprepGraphAct
 {
 	bNodeDragging = true;
 	bSkipNextDragUpdate = false;
- 
+	bSkipRefreshLayout = false;
+	bCursorLeftOnLeft = false;
+	bCursorLeftOnRight = false;
+	bNodeDragJustStarted = true;
+	LastDragDirection = 0.f;
+
 	TrackWidgetPtr->OnStartNodeDrag(ActionNode->GetExecutionOrder());
 
-	LastDragScreenSpacePosition = FSlateApplication::Get().GetCursorPos();
-	const FGeometry& DesktopGeometry = GetOwnerPanel()->GetPersistentState().DesktopGeometry;
-	FVector2D DragLocalPosition = DesktopGeometry.AbsoluteToLocal( LastDragScreenSpacePosition );
+	const FGeometry& PanelGeometry = GetOwnerPanel()->GetTickSpaceGeometry();
+
+	const FVector2D PanelPosition = PanelGeometry.AbsolutePosition;
+	const FVector2D PanelSize = PanelGeometry.GetLocalSize();
+	RECT Boundaries;
+	Boundaries.left = FMath::RoundToInt(PanelPosition.X);
+	Boundaries.top = FMath::RoundToInt(PanelPosition.Y);
+	Boundaries.right = Boundaries.left + FMath::RoundToInt(PanelSize.X);
+	Boundaries.bottom = Boundaries.top + FMath::RoundToInt(PanelSize.Y);
+
+	TSharedPtr<ICursor> PlatformCursor = FSlateApplication::Get().GetPlatformCursor();
+	PlatformCursor->SetType(EMouseCursor::None);
 
 	SGraphPanel* GraphPanel = OwnerGraphPanelPtr.Pin().Get();
 	ensure(GraphPanel);
@@ -588,121 +647,237 @@ void SDataprepGraphTrackNode::OnStartNodeDrag(const TSharedRef<SDataprepGraphAct
 	GraphPanel->Invalidate(EInvalidateWidgetReason::Layout);
 }
 
-int32 SDataprepGraphTrackNode::OnEndNodeDrag()
+void SDataprepGraphTrackNode::OnEndNodeDrag(bool bDoDrop)
 {
 	bNodeDragging = false;
 	bSkipNextDragUpdate = false;
+	bSkipRefreshLayout = true;
+	bCursorLeftOnLeft = false;
+	bCursorLeftOnRight = false;
+	bNodeDragJustStarted = false;
+	LastDragDirection = 0.f;
 
 	SGraphPanel* GraphPanel = OwnerGraphPanelPtr.Pin().Get();
 	ensure(GraphPanel);
 
 	GraphPanel->Invalidate(EInvalidateWidgetReason::Layout);
 
-	return 	TrackWidgetPtr->OnEndNodeDrag();
+
+	FVector2D NewScreenSpacePosition;
+	int32 DraggedIndex;
+	int32 DroppedIndex;
+
+	TrackWidgetPtr->OnEndNodeDrag(NewScreenSpacePosition, DraggedIndex, DroppedIndex);
+
+	FSlateApplication::Get().GetPlatformCursor()->SetType(EMouseCursor::Default);
+	FSlateApplication::Get().SetCursorPos( NewScreenSpacePosition );
+
+	if(bDoDrop)
+	{
+		if(UDataprepAsset* DataprepAsset = GetDataprepAsset())
+		{
+			FModifierKeysState ModifierKeyState = FSlateApplication::Get().GetModifierKeys();
+			const bool bCopyRequested = ModifierKeyState.IsControlDown() || ModifierKeyState.IsCommandDown();
+
+			FScopedTransaction Transaction( bCopyRequested ? LOCTEXT("DataprepGraphTrackNode_OnDropCopy", "Add/Insert action") : LOCTEXT("OnDropMove", "Move action") );
+			bool bTransactionSuccessful = false;
+
+			if(bCopyRequested)
+			{
+				if(DroppedIndex >= DataprepAsset->GetActionCount())
+				{
+					bTransactionSuccessful = DataprepAsset->AddAction(DataprepAsset->GetAction(DraggedIndex)) != INDEX_NONE;
+				}
+				else
+				{
+					bTransactionSuccessful = DataprepAsset->InsertAction(DataprepAsset->GetAction(DraggedIndex), DroppedIndex);
+				}
+			}
+			else if( DroppedIndex != DraggedIndex)
+			{
+				bTransactionSuccessful = DataprepAsset->MoveAction(DraggedIndex, DroppedIndex);
+			}
+			else
+			{
+				RefreshLayout();
+			}
+
+			if(!bTransactionSuccessful)
+			{
+				Transaction.Cancel();
+			}
+		}
+	}
+	else
+	{
+		RefreshLayout();
+	}
 }
 
-void SDataprepGraphTrackNode::OnNodeDragged( TSharedPtr<SDataprepGraphActionNode>& ActionNodePtr, const FVector2D& DragScreenSpacePosition, const FVector2D& ScreenSpaceDelta)
+bool SDataprepGraphTrackNode::OnNodeDragged( TSharedPtr<SDataprepGraphActionNode>& ActionNodePtr, const FVector2D& DragScreenSpacePosition, const FVector2D& InScreenSpaceDelta)
 {
 	ensure(bNodeDragging);
 	ensure(ActionNodePtr.IsValid());
 
-	// This update is most likely due to a call to FSlateApplication::SetCursorPos. Skip it
-	if(bSkipNextDragUpdate || ScreenSpaceDelta.X == 0.f)
+	// Do not proceed until the user has actually move the mouse
+	if(bNodeDragJustStarted && InScreenSpaceDelta.X == 0.f)
+	{
+		return TrackWidgetPtr->bMouseInScope;
+	}
+
+	if(InScreenSpaceDelta.X != 0.f)
 	{
+		LastDragDirection = FMath::Sign(InScreenSpaceDelta.X);
+	}
+
+	bNodeDragJustStarted = false;
+	LastSetCursorPosition = TrackWidgetPtr->GetAbsoluteSoftwareCursorPosition();
+
+	FVector2D ScreenSpaceDelta = InScreenSpaceDelta;
+
+	// The mouse was out of the panel and is back
+	if(bSkipNextDragUpdate)
+	{
+		const FVector2D PreviousCursorPosition = DragScreenSpacePosition + ScreenSpaceDelta;
+		if(PreviousCursorPosition.Equals(LastSetCursorPosition, 0.9f) || DragScreenSpacePosition.Equals(LastSetCursorPosition, 0.9f))
+		{
+			return TrackWidgetPtr->bMouseInScope;
+		}
+
 		bSkipNextDragUpdate = false;
-		return;
 	}
 
 	TSharedPtr<SGraphPanel> GraphPanel = GetOwnerPanel();
-	const FGeometry& DesktopGeometry = GraphPanel->GetCachedGeometry();
-	const FVector2D& PanelSize = DesktopGeometry.GetLocalSize();
-	const float ZoomAmount = GraphPanel->GetZoomAmount();
-	const FVector2D& TrackSize = TrackWidgetPtr->GetWorkingSize() * ZoomAmount;
+	const FGeometry& PanelGeometry = GraphPanel->GetTickSpaceGeometry();
+	const FVector2D PanelSize = PanelGeometry.GetLocalSize();
 
-	float AbscissaDelta = TrackWidgetPtr->OnNodeDragged(ScreenSpaceDelta.X / ZoomAmount);
-		
-	FVector2D NewScreenSpaceDelta( AbscissaDelta * ZoomAmount, ScreenSpaceDelta.Y );
+	// If pointer is leaving panel's window, make sure it is visible and skip processing
+	const float PanelMinX = PanelGeometry.AbsolutePosition.X;
+	const float PanelMaxX = PanelMinX + PanelSize.X;
 
-	// Request the active panel to scroll if required
+	if(DragScreenSpacePosition.X < PanelMinX || DragScreenSpacePosition.X > PanelMaxX)
 	{
-		// Keep the mouse at the same position if dragged node has reached when of the ends of the track
-		if(AbscissaDelta == 0.f)
+		if(TrackWidgetPtr->bMouseInScope)
 		{
-			// Do nothing, the mouse's cursor and the dragged node should not move
+			TrackWidgetPtr->bMouseInScope = false;
+			bCursorLeftOnLeft = DragScreenSpacePosition.X < PanelMinX;
+			bCursorLeftOnRight = DragScreenSpacePosition.X > PanelMaxX;
+
+			TSharedPtr<FSlateUser> SlateUser = FSlateApplication::Get().GetCursorUser();
+			// #ueent_todo: Is call to SetCursorVisibility useful?
+			SlateUser->SetCursorVisibility(true);
+			SlateUser->SetCursorPosition(FVector2D(DragScreenSpacePosition.X, LastSetCursorPosition.Y));
 		}
-		// Panel is narrower than track, update panel if dragged node has entered non visible part of track
-		else if(TrackSize.X > PanelSize.X)
+
+		return TrackWidgetPtr->bMouseInScope;
+	}
+
+	const FGeometry& TrackGeometry = TrackWidgetPtr->GetTickSpaceGeometry();
+	FVector2D DragTrackPosition = TrackGeometry.AbsoluteToLocal(DragScreenSpacePosition);
+	FVector2D SoftwareCursorPos = TrackWidgetPtr->GetSoftwareCursorPosition();
+
+	// Mouse pointer was out of panel's screen space, make sure to hide mouse cursor
+	TrackWidgetPtr->bMouseInScope = true;
+
+	// When cursor was out of scope, do not process drag until cursor reaches marker
+	if(bCursorLeftOnLeft || bCursorLeftOnRight)
+	{
+		const bool bStillOutOfScope = bCursorLeftOnLeft && DragTrackPosition.X < SoftwareCursorPos.X ||
+									  bCursorLeftOnRight && DragTrackPosition.X > SoftwareCursorPos.X;
+
+		// Still has not passed track marker, skip drag processing;
+		if(bStillOutOfScope)
 		{
-			FVector2D LeftPosition;
-			FVector2D RightPosition;
-			TrackWidgetPtr->GetDraggedNodePositions(LeftPosition, RightPosition, DesktopGeometry);
+			return TrackWidgetPtr->bMouseInScope;
+		}
+
+		// Adjust delta to only contain move from or to track marker's position
+		ScreenSpaceDelta.X = bCursorLeftOnLeft ? DragTrackPosition.X + ScreenSpaceDelta.X - SoftwareCursorPos.X : DragTrackPosition.X - SoftwareCursorPos.X;
 
-			if(RightPosition.X > PanelSize.X)
+		bCursorLeftOnLeft = false;
+		bCursorLeftOnRight = false;
+	}
+
+	const float ZoomAmount = GraphPanel->GetZoomAmount();
+	const FVector2D TrackSize = TrackWidgetPtr->GetWorkingSize() * ZoomAmount;
+	float PanAmount = 0.f;
+
+	DragCallback LocalComputePanAmount = [&](const FVector2D& SlotRange, float Delta) -> float
+	{
+		if(TrackSize.X > PanelSize.X)
+		{
+			const float MaxPanSpeed = /*ZoomAmount < 1.f ? 10.f : */10.f / ZoomAmount ;
+			FVector2D ScreenSpacePosition = TrackGeometry.LocalToAbsolute(FVector2D(SlotRange.X, DragTrackPosition.Y));
+			PanAmount = this->ComputePanAmount(ScreenSpacePosition, MaxPanSpeed).X;
+
+			ScreenSpacePosition = TrackGeometry.LocalToAbsolute(FVector2D(SlotRange.Y, DragTrackPosition.Y));
+			PanAmount += this->ComputePanAmount(ScreenSpacePosition, MaxPanSpeed).X;
+
+			// Adjust pan amount to boundaries of the track widget
+			const float CurrentViewOffsetX = GraphPanel->GetViewOffset().X * ZoomAmount;
+			const float NewViewOffsetX = (GraphPanel->GetViewOffset().X + PanAmount) * ZoomAmount;
+			if(NewViewOffsetX < 0.f)
 			{
-				FVector2D ViewOffset = GraphPanel->GetViewOffset();
-				ViewOffset.X += RightPosition.X - PanelSize.X + 10.f;
-				GraphPanel->RestoreViewSettings(ViewOffset, ZoomAmount);
+				PanAmount = -GraphPanel->GetViewOffset().X;
 			}
-			else if(LeftPosition.X < 0.f)
+			else if(NewViewOffsetX + PanelSize.X > TrackSize.X)
 			{
-				FVector2D ViewOffset = GraphPanel->GetViewOffset();
-				ViewOffset.X += LeftPosition.X - 10.f;
-				GraphPanel->RestoreViewSettings(ViewOffset, ZoomAmount);
+				PanAmount = (TrackSize.X - (CurrentViewOffsetX + PanelSize.X)) / ZoomAmount;
+				if(PanAmount < 0.f)
+				{
+					PanAmount = 0.f;
+				}
 			}
-			else
+
+			// Preventing panning from going at the opposite of current mouse move
+			if(PanAmount * LastDragDirection < 0.f)
 			{
-				// Make sure cursor stays at a constant height
-				LastDragScreenSpacePosition.X = DragScreenSpacePosition.X;
+				PanAmount = 0.f;
 			}
 		}
-		else
-		{
-			// Make sure cursor stays at a constant height
-			LastDragScreenSpacePosition.X = DragScreenSpacePosition.X;
-		}
 
-		// Update cursor's position
-		bSkipNextDragUpdate = LastDragScreenSpacePosition.X != DragScreenSpacePosition.X;
-		FSlateApplication::Get().SetCursorPos( LastDragScreenSpacePosition );
+		return PanAmount;
+	};
+	
+	TrackWidgetPtr->OnNodeDragged(ScreenSpaceDelta.X / ZoomAmount, LocalComputePanAmount);
+
+	if(PanAmount != 0.f)
+	{
+		bSkipRefreshLayout = true;
+		GraphPanel->RestoreViewSettings(GraphPanel->GetViewOffset() + (FVector2D(PanAmount,0.f) / ZoomAmount), ZoomAmount);
 	}
+
+	return true;
 }
 
 void SDataprepGraphTrackNode::OnControlKeyChanged(bool bControlKeyDown)
 {
 	if(bNodeDragging)
 	{
-		if(bControlKeyDown)
-		{
-			TrackWidgetPtr->OnControlKeyDown(true);
-		}
-		else
-		{
-			float MouseOffset = TrackWidgetPtr->OnControlKeyDown(false);
-
-			if(MouseOffset != 0.f)
-			{
-				// Update mouse position
-				LastDragScreenSpacePosition.X -= MouseOffset;
+		FVector2D NewScreenSpacePosition;
+		TrackWidgetPtr->OnControlKeyDown(bControlKeyDown, NewScreenSpacePosition);
 
-				bSkipNextDragUpdate = true;
-				FSlateApplication::Get().SetCursorPos( LastDragScreenSpacePosition );
-			}
+		if(!bControlKeyDown && NewScreenSpacePosition != FVector2D::ZeroVector)
+		{
+			bSkipNextDragUpdate = true;
+			FSlateApplication::Get().SetCursorPos( NewScreenSpacePosition );
 		}
 	}
 }
 
-void SDataprepGraphTrackNode::RequestViewportPan(const FVector2D& InScreenSpacePosition)
+FVector2D SDataprepGraphTrackNode::ComputePanAmount(const FVector2D& InScreenSpacePosition, float MaxPanSpeed)
 {
 	// Note: Code copied from SNodePanel::RequestDeferredPan and its subsequent calls. Calling SNodePanel::RequestDeferredPan is not an option
 	//		 since it changes the view offset during the call to OnPaint which bypasses the adjustments done by the SDataprepGraphEditor in
 	//		 SDataprepGraphEditor::Tick 
+	FVector2D PanAmount(FVector2D::ZeroVector);
+
 	if(SGraphPanel* GraphPanel = OwnerGraphPanelPtr.Pin().Get())
 	{
 		// How quickly to ramp up the pan speed as the user moves the mouse further past the edge of the graph panel.
 		constexpr float EdgePanSpeedCoefficient = 2.f;
 		constexpr float EdgePanSpeedPower = 0.6f;
 		// Never pan faster than this - probably not really required since we raise to a power of 0.6
-		constexpr float MaxPanSpeed = 200.0f;
 		// Start panning before we reach the edge of the graph panel.
 		constexpr float EdgePanForgivenessZone = 30.0f;
 
@@ -712,7 +887,6 @@ void SDataprepGraphTrackNode::RequestViewportPan(const FVector2D& InScreenSpaceP
 		const FVector2D PanelSize = PanelGeometry.GetLocalSize();
 		const float ZoomAmount = GraphPanel->GetZoomAmount();
 
-		FVector2D PanAmount(FVector2D::ZeroVector);
 		if ( PanelPosition.X <= EdgePanForgivenessZone )
 		{
 			PanAmount.X = FMath::Max( -MaxPanSpeed, EdgePanSpeedCoefficient * -FMath::Pow(EdgePanForgivenessZone - PanelPosition.X, EdgePanSpeedPower) );
@@ -730,10 +904,23 @@ void SDataprepGraphTrackNode::RequestViewportPan(const FVector2D& InScreenSpaceP
 		{
 			PanAmount.Y = FMath::Min( MaxPanSpeed, EdgePanSpeedCoefficient * FMath::Pow(PanelPosition.Y - PanelSize.Y + EdgePanForgivenessZone, EdgePanSpeedPower) );
 		}
+	}
+
+	return PanAmount;
+}
+
+void SDataprepGraphTrackNode::RequestViewportPan(const FVector2D& InScreenSpacePosition)
+{
+	if(SGraphPanel* GraphPanel = OwnerGraphPanelPtr.Pin().Get())
+	{
+		FVector2D PanAmount = ComputePanAmount(InScreenSpacePosition);
 
 		if(PanAmount != FVector2D::ZeroVector)
 		{
-			GraphPanel->RestoreViewSettings(GraphPanel->GetViewOffset() + PanAmount / ZoomAmount, ZoomAmount);
+			const float ZoomAmount = GraphPanel->GetZoomAmount();
+
+			bSkipRefreshLayout = true;
+			GraphPanel->RestoreViewSettings(GraphPanel->GetViewOffset() + (PanAmount / ZoomAmount), ZoomAmount);
 		}
 	}
 }
@@ -773,10 +960,21 @@ void SDataprepGraphTrackNode::RequestRename(const UEdGraphNode* Node)
 	}
 }
 
+void SDataprepGraphTrackNode::UpdateProxyNode(TSharedRef<SDataprepGraphActionNode> ActioNodePtr, const FVector2D& ScreenSpacePosition)
+{
+	if(SGraphPanel* GraphPanel = OwnerGraphPanelPtr.Pin().Get())
+	{
+		FVector2D PanelPosition = GraphPanel->GetTickSpaceGeometry().AbsoluteToLocal(ScreenSpacePosition);
+		FVector2D GraphPosition = PanelPosition / GraphPanel->GetZoomAmount() + GraphPanel->GetViewOffset();
+		ActioNodePtr->UpdateProxyNode(GraphPosition);
+	}
+}
+
 void SDataprepGraphTrackWidget::Construct(const FArguments& InArgs, TSharedPtr<SDataprepGraphTrackNode> InTrackNode)
 {
 	TrackNodePtr = InTrackNode;
 	ensure(TrackNodePtr.IsValid());
+	bMouseInScope = false;
 
 	SConstraintCanvas::Construct( SConstraintCanvas::FArguments());
 
@@ -974,9 +1172,8 @@ void SDataprepGraphTrackWidget::Construct(const FArguments& InArgs, TSharedPtr<S
 	WorkingSize = FVector2D::ZeroVector;
 	CanvasOffset.Right = LeftOffset;
 
-	HoveredDragIndex = INDEX_NONE;
-	DragStartIndex = INDEX_NONE;
-	DragIndex = INDEX_NONE;
+	IndexHovered = INDEX_NONE;
+	IndexDragged = INDEX_NONE;
 	DragIndicatorIndex = INDEX_NONE;
 }
 
@@ -1042,35 +1239,6 @@ void SDataprepGraphTrackWidget::CreateHelperWidgets()
 	DummyAction = SNew(SDataprepEmptyActionNode);
 }
 
-void SDataprepGraphTrackWidget::GetDraggedNodePositions(FVector2D& Left, FVector2D& Right, const FGeometry& TargetGeometry)
-{
-	const FMargin DragSlotOffset = DragSlot->OffsetAttr.Get(FMargin());
-	const FGeometry& Geometry = GetCachedGeometry();
-
-	Left = TargetGeometry.AbsoluteToLocal(Geometry.LocalToAbsolute(FVector2D(DragSlotOffset.Left, 0.f)));
-	Right = TargetGeometry.AbsoluteToLocal(Geometry.LocalToAbsolute(FVector2D(DragSlotOffset.Left + DragSlotOffset.Right, 0.f)));
-}
-
-void SDataprepGraphTrackWidget::SetActionSlot(int32 SlotIndex, const TSharedRef<SWidget>& Widget)
-{
-	ensure(ActionSlots.IsValidIndex(SlotIndex));
-	ActionSlots[SlotIndex]->AttachWidget(Widget);
-
-	Invalidate(EInvalidateWidgetReason::Layout);
-}
-
-FVector2D SDataprepGraphTrackWidget::GetActioNodeAbscissaRange()
-{
-	FVector2D Range(InterSlotSpacing, InterSlotSpacing + SDataprepGraphActionNode::DefaultWidth + InterSlotSpacing);
-
-	if(SDataprepGraphTrackNode* TrackNode = TrackNodePtr.Pin().Get())
-	{
-		Range.Y = ActionSlots[SlotCount-1]->OffsetAttr.Get(FMargin()).Left;
-	}
-
-	return Range;
-}
-
 int32 SDataprepGraphTrackWidget::GetHoveredActionNode(float LeftCorner, float Width)
 {
 	const float CenterAbscissa = LeftCorner + Width * 0.5f;
@@ -1115,7 +1283,9 @@ void SDataprepGraphTrackWidget::RefreshLayout()
 	for(int32 Index = 0; Index < SlotCount; ++Index)
 	{
 		// Update action's proxy node registered to graph panel
-		StaticCastSharedRef<SDataprepGraphActionNode>(ActionNodes[Index])->UpdateProxyNode(FVector2D(DropSlotOffset.Left + InterSlotSpacing, NodeTopPadding) + TrackOffset);
+		const FVector2D LocalNodePosition(DropSlotOffset.Left + InterSlotSpacing, NodeTopPadding);
+		const FVector2D AbsoluteNodePosition = GetTickSpaceGeometry().LocalToAbsolute(LocalNodePosition);
+		TrackNodePtr.Pin()->UpdateProxyNode(StaticCastSharedRef<SDataprepGraphActionNode>(ActionNodes[Index]), AbsoluteNodePosition);
 
 		const FVector2D Size = UpdateActionSlot(DropSlotOffset, Index, ActionNodes[Index]);
 
@@ -1148,69 +1318,69 @@ void SDataprepGraphTrackWidget::RefreshLayout()
 	Invalidate(EInvalidateWidgetReason::Layout | EInvalidateWidgetReason::ChildOrder);
 }
 
-void SDataprepGraphTrackWidget::UpdateLayoutOnDrag(int8 DragDirection)
+void SDataprepGraphTrackWidget::UpdateLayoutOnDrag()
 {
 	FMargin SlotOffset = DropSlots[0]->OffsetAttr.Get(FMargin());
 
 	if(bIsCopying)
 	{
-		for(int32 Index = 0; Index < HoveredDragIndex; ++Index)
+		for(int32 Index = 0; Index < IndexHovered; ++Index)
 		{
 			UpdateActionSlot(SlotOffset, Index, ActionNodes[Index]);
 		}
 
-		UpdateActionSlot(SlotOffset, HoveredDragIndex, DragSlot->GetWidget(), true);
+		UpdateActionSlot(SlotOffset, IndexHovered, DragSlot->GetWidget(), true);
 
-		for(int32 Index = HoveredDragIndex + 1; Index < SlotCount; ++Index)
+		for(int32 Index = IndexHovered + 1; Index < SlotCount; ++Index)
 		{
 			UpdateActionSlot(SlotOffset, Index, ActionNodes[Index-1]);
 		}
 	}
-	else if(HoveredDragIndex < DragStartIndex)
+	else if(IndexHovered < IndexDragged)
 	{
-		for(int32 Index = 0; Index < HoveredDragIndex; ++Index)
+		for(int32 Index = 0; Index < IndexHovered; ++Index)
 		{
 			UpdateActionSlot(SlotOffset, Index, ActionNodes[Index]);
 		}
 
-		UpdateActionSlot(SlotOffset, HoveredDragIndex, DragSlot->GetWidget(), true);
+		UpdateActionSlot(SlotOffset, IndexHovered, DragSlot->GetWidget(), true);
 
-		for(int32 Index = HoveredDragIndex + 1; Index <= DragStartIndex; ++Index)
+		for(int32 Index = IndexHovered + 1; Index <= IndexDragged; ++Index)
 		{
 			UpdateActionSlot(SlotOffset, Index, ActionNodes[Index-1]);
 		}
 
-		for(int32 Index = DragStartIndex+1; Index < SlotCount; ++Index)
+		for(int32 Index = IndexDragged+1; Index < SlotCount; ++Index)
 		{
 			UpdateActionSlot(SlotOffset, Index, ActionNodes[Index]);
 		}
 	}
 	else
 	{
-		for(int32 Index = 0; Index < DragStartIndex; ++Index)
+		for(int32 Index = 0; Index < IndexDragged; ++Index)
 		{
 			UpdateActionSlot(SlotOffset, Index, ActionNodes[Index]);
 		}
 
-		for(int32 Index = DragStartIndex; Index < HoveredDragIndex; ++Index)
+		for(int32 Index = IndexDragged; Index < IndexHovered; ++Index)
 		{
 			UpdateActionSlot(SlotOffset, Index, ActionNodes[Index+1]);
 		}
 
-		UpdateActionSlot(SlotOffset, HoveredDragIndex, DragSlot->GetWidget(), true);
+		UpdateActionSlot(SlotOffset, IndexHovered, DragSlot->GetWidget(), true);
 
-		for(int32 Index = HoveredDragIndex + 1; Index < SlotCount; ++Index)
+		for(int32 Index = IndexHovered + 1; Index < SlotCount; ++Index)
 		{
 			UpdateActionSlot(SlotOffset, Index, ActionNodes[Index]);
 		}
 	}
 
 	const FMargin DragSlotOffset = DragSlot->OffsetAttr.Get(FMargin());
-	const FMargin HoveredDropSlotOffset = DropSlots[HoveredDragIndex]->OffsetAttr.Get(FMargin());
+	const FMargin HoveredDropSlotOffset = DropSlots[IndexHovered]->OffsetAttr.Get(FMargin());
 
-	if(HoveredDragIndex > 0 && DragSlotOffset.Left < HoveredDropSlotOffset.Left)
+	if(IndexHovered > 0 && DragSlotOffset.Left < HoveredDropSlotOffset.Left)
 	{
-		SConstraintCanvas::FSlot* OverlappedSlot = ActionSlots[HoveredDragIndex-1];
+		SConstraintCanvas::FSlot* OverlappedSlot = ActionSlots[IndexHovered-1];
 		FMargin OverlappedSlotOffset = OverlappedSlot->OffsetAttr.Get(FMargin());
 
 		const float OverlappedActionSlotLeft = HoveredDropSlotOffset.Left - OverlappedSlotOffset.Right;
@@ -1225,11 +1395,11 @@ void SDataprepGraphTrackWidget::UpdateLayoutOnDrag(int8 DragDirection)
 			OverlappedSlot->OffsetAttr.Set(OverlappedSlotOffset);
 		}
 	}
-	else if(HoveredDragIndex < (SlotCount-1) && DragSlotOffset.Left > (HoveredDropSlotOffset.Left + HoveredDropSlotOffset.Right) )
+	else if(IndexHovered < (SlotCount-1) && DragSlotOffset.Left > (HoveredDropSlotOffset.Left + HoveredDropSlotOffset.Right) )
 	{
-		SConstraintCanvas::FSlot* OverlappedSlot = ActionSlots[HoveredDragIndex+1];
+		SConstraintCanvas::FSlot* OverlappedSlot = ActionSlots[IndexHovered+1];
 		FMargin OverlappedSlotOffset = OverlappedSlot->OffsetAttr.Get(FMargin());
-		FMargin OverlappedDropSlotOffset = DropSlots[HoveredDragIndex+1]->OffsetAttr.Get(FMargin());
+		FMargin OverlappedDropSlotOffset = DropSlots[IndexHovered+1]->OffsetAttr.Get(FMargin());
 
 		const float OverlappedActionSlotLeft = OverlappedDropSlotOffset.Left + OverlappedDropSlotOffset.Right;
 		const float MinimalLeftOffset = HoveredDropSlotOffset.Left + HoveredDropSlotOffset.Right;
@@ -1248,7 +1418,7 @@ void SDataprepGraphTrackWidget::UpdateLayoutOnDrag(int8 DragDirection)
 	Invalidate(EInvalidateWidgetReason::Layout | EInvalidateWidgetReason::Paint);
 }
 
-void SDataprepGraphTrackWidget::OnStartNodeDrag(int32 IndexDragged)
+void SDataprepGraphTrackWidget::OnStartNodeDrag(int32 InIndexDragged)
 {
 	if(!DragSlot)
 	{
@@ -1257,76 +1427,103 @@ void SDataprepGraphTrackWidget::OnStartNodeDrag(int32 IndexDragged)
 	}
 
 	ensure(TrackNodePtr.IsValid());
-	ensure(ActionSlots.IsValidIndex(IndexDragged));
+	ensure(ActionSlots.IsValidIndex(InIndexDragged));
 
-	DragStartIndex = IndexDragged;
-	DragIndex = IndexDragged;
-	HoveredDragIndex = IndexDragged;
+	IndexDragged = InIndexDragged;
+	IndexHovered = InIndexDragged;
 	DragIndicatorIndex = INDEX_NONE;
 
 	FMargin Offset = ActionSlots[IndexDragged]->OffsetAttr.Get(FMargin());
 	DragSlot->OffsetAttr.Set(Offset);
 	DragSlot->AttachWidget(ActionNodes[IndexDragged]);
 
+	const FVector2D LocalCursorPosition = GetTickSpaceGeometry().AbsoluteToLocal(FSlateApplication::Get().GetCursorPos());
+	TrackCursorOffset = FVector2D(LocalCursorPosition.X - Offset.Left, LocalCursorPosition.Y);
+
 	AbscissaRange.X = ActionSlots[0]->OffsetAttr.Get(FMargin()).Left;
+	// AbscissaRange.Y will be properly update in the call to OnControlKeyDown
+	AbscissaRange.Y = AbscissaRange.X;
 
 	const FModifierKeysState ModifierKeyState = FSlateApplication::Get().GetModifierKeys();
 	bIsCopying = !(ModifierKeyState.IsControlDown() || ModifierKeyState.IsCommandDown());
 
-	OnControlKeyDown(!bIsCopying);
+	FVector2D NewScreenSpacePosition;
+	OnControlKeyDown(!bIsCopying, NewScreenSpacePosition);
 }
 
-float SDataprepGraphTrackWidget::OnNodeDragged(float DragDelta)
+void SDataprepGraphTrackWidget::OnNodeDragged(float DragDelta, DragCallback ComputePanAmount)
 {
-	float ActualDelta = 0.f;
-
 	if(!DragSlot)
 	{
 		ensure(false);
-		return ActualDelta;
+		return;
 	}
 
 	FMargin DragSlotOffset = DragSlot->OffsetAttr.Get(FMargin());
 
-	float NewAbscissa = DragSlotOffset.Left + DragDelta;
+	const float DesiredAbscissa = DragSlotOffset.Left + DragDelta;
 	 
-	bool bValidDrag = (NewAbscissa >= AbscissaRange.X) && (DragSlotOffset.Left < AbscissaRange.Y || NewAbscissa < AbscissaRange.Y);
+	bool bValidDrag = (DesiredAbscissa >= AbscissaRange.X) && (DragSlotOffset.Left < AbscissaRange.Y || DesiredAbscissa < AbscissaRange.Y);
 
 	if(bValidDrag)
 	{
-		TSharedPtr<SDataprepGraphActionNode> DraggedActionNode = StaticCastSharedRef<SDataprepGraphActionNode>(DragSlot->GetWidget());
-		ensure(DraggedActionNode.IsValid());
+		float NodeNewAbscissa = DesiredAbscissa < AbscissaRange.X ? AbscissaRange.X : (DesiredAbscissa > AbscissaRange.Y ? AbscissaRange.Y : DesiredAbscissa);
 
-		const float CachedNodeAbscissa = DragSlotOffset.Left;
+		const FVector2D DraggedSlotSize(DragSlotOffset.Left, DragSlotOffset.Left + DragSlotOffset.Right);
+		const float MouseDelta = NodeNewAbscissa - DragSlotOffset.Left;
 
-		DragSlotOffset.Left = NewAbscissa < AbscissaRange.X ? AbscissaRange.X : (NewAbscissa > AbscissaRange.Y ? AbscissaRange.Y : NewAbscissa);
-		ActualDelta = DragSlotOffset.Left - CachedNodeAbscissa;
+		NodeNewAbscissa += ComputePanAmount(DraggedSlotSize, MouseDelta);
 
-		DragSlot->OffsetAttr.Set(DragSlotOffset);
-		DraggedActionNode->GetNodeObj()->NodePosX = DragSlotOffset.Left;
+		// Keep slot's abscissa within range
+		if(NodeNewAbscissa < AbscissaRange.X)
+		{
+			NodeNewAbscissa = AbscissaRange.X;
+		}
+		else if(NodeNewAbscissa > AbscissaRange.Y)
+		{
+			NodeNewAbscissa = AbscissaRange.Y;
+		}
+
+		if(NodeNewAbscissa != DragSlotOffset.Left)
+		{
+			TSharedPtr<SDataprepGraphActionNode> DraggedActionNode = StaticCastSharedRef<SDataprepGraphActionNode>(DragSlot->GetWidget());
+			ensure(DraggedActionNode.IsValid());
+
+			// Add offset from current position to related GraphNode
+			DraggedActionNode->GetNodeObj()->NodePosX += NodeNewAbscissa - DragSlotOffset.Left;
 
-		// Check if center of dragged widget is over a neighboring widget by at least half its size
-		HoveredDragIndex = GetHoveredActionNode(DragSlotOffset.Left, DragSlotOffset.Right);
-		UpdateLayoutOnDrag(ActualDelta < 0.f ? -1 : (ActualDelta > 0.f ? 1 : 0));
+			// Update drag slot
+			DragSlotOffset.Left = NodeNewAbscissa;
+			DragSlot->OffsetAttr.Set(DragSlotOffset);
+
+			// Check if center of dragged widget is over a neighboring widget by at least half its size
+			IndexHovered = GetHoveredActionNode(DragSlotOffset.Left, DragSlotOffset.Right);
+			UpdateLayoutOnDrag();
+		}
 
 		Invalidate(EInvalidateWidgetReason::Layout | EInvalidateWidgetReason::ChildOrder);
 	}
-
-	return ActualDelta;
 }
 
-int32 SDataprepGraphTrackWidget::OnEndNodeDrag()
+void SDataprepGraphTrackWidget::OnEndNodeDrag(FVector2D& OutScreenSpacePosition, int32& OutDraggedIndex, int32&OutDroppedIndex)
 {
 	if(!DragSlot)
 	{
 		ensure(false);
-		return INDEX_NONE;
+		return;
 	}
 
+	// Retrieve all data potentially useful for the drop
+	OutScreenSpacePosition = bMouseInScope ? GetAbsoluteSoftwareCursorPosition() : FSlateApplication::Get().GetCursorPos();
+	OutDraggedIndex = IndexDragged;
+	OutDroppedIndex = IndexHovered;
+
 	TSharedPtr<SDataprepGraphActionNode> ActionNode = StaticCastSharedRef<SDataprepGraphActionNode>(DragSlot->GetWidget());
 	ensure(ActionNode.IsValid());
 
+	// Reset all  members used during the drag
 	bIsCopying = false;
+	bMouseInScope = false;
 
 	SlotCount = ActionNodes.Num();
 
@@ -1336,16 +1533,15 @@ int32 SDataprepGraphTrackWidget::OnEndNodeDrag()
 	DropSlots[SlotCount]->AttachWidget(SNullWidget::NullWidget);
 	ActionSlots[SlotCount]->AttachWidget(SNullWidget::NullWidget);
 
-	RefreshLayout();
+	IndexDragged = INDEX_NONE;
+	IndexHovered = INDEX_NONE;
 
-	int32 CachedHoveredDragIndex = HoveredDragIndex;
-	DragStartIndex = INDEX_NONE;
-	DragIndex = INDEX_NONE;
-	HoveredDragIndex = INDEX_NONE;
+	TrackCursorOffset = FVector2D::ZeroVector;
 
-	Invalidate(EInvalidateWidgetReason::Layout | EInvalidateWidgetReason::ChildOrder);
+	// Update the widget
+	RefreshLayout();
 
-	return CachedHoveredDragIndex;
+	Invalidate(EInvalidateWidgetReason::Layout | EInvalidateWidgetReason::ChildOrder);
 }
 
 void SDataprepGraphTrackWidget::UpdateDragIndicator(FVector2D MousePosition)
@@ -1407,26 +1603,27 @@ void SDataprepGraphTrackWidget::UpdateDragIndicator(FVector2D MousePosition)
 	}
 }
 
-float SDataprepGraphTrackWidget::OnControlKeyDown(bool bCtrlDown)
+void SDataprepGraphTrackWidget::OnControlKeyDown(bool bCtrlDown, FVector2D& OutScreenSpacePosition)
 {
+	OutScreenSpacePosition = FVector2D::ZeroVector;
+
 	if(bIsCopying == bCtrlDown)
 	{
-		return 0.f;
+		return;
 	}
 
-	if(DragStartIndex == INDEX_NONE)
+	if(IndexDragged == INDEX_NONE)
 	{
 		bIsCopying = false;
-		return 0.f;
+		return;
 	}
 
 	bIsCopying = bCtrlDown;
 
-	float MouseOffset = 0.f;
 	SlotCount = ActionNodes.Num();
 
 	FMargin DragSlotOffset = DragSlot->OffsetAttr.Get(FMargin());
-	HoveredDragIndex = GetHoveredActionNode(DragSlotOffset.Left, DragSlotOffset.Right);
+	IndexHovered = GetHoveredActionNode(DragSlotOffset.Left, DragSlotOffset.Right);
 
 	for(int32 Index = 0; Index < SlotCount; ++Index)
 	{
@@ -1435,55 +1632,44 @@ float SDataprepGraphTrackWidget::OnControlKeyDown(bool bCtrlDown)
 
 	if(bIsCopying)
 	{
-		// Set the maximum move on the right
-		const FVector2D DragSize = GetNodeSize(DragSlot->GetWidget());
-		FMargin LastDropSlotOffset = DropSlots[SlotCount]->OffsetAttr.Get(FMargin());
-		AbscissaRange.Y = LastDropSlotOffset.Left + LastDropSlotOffset.Right + 1.f - (DragSize.X * 0.5f);
-
 		DropSlots[SlotCount]->AttachWidget(DropFiller.ToSharedRef());
 		ActionSlots[SlotCount]->AttachWidget(DummyAction.ToSharedRef());
 		TrackSlots[SlotCount]->AttachWidget(TrackSlotRegular.ToSharedRef());
 
+		// Set the maximum move on the right and new working size
+		const FVector2D DragSize = GetNodeSize(DragSlot->GetWidget());
+		const FMargin SlotOffset = TrackSlots[SlotCount]->OffsetAttr.Get(FMargin());
+
+		AbscissaRange.Y = SlotOffset.Left - (DragSize.X * 0.5f) + InterSlotSpacing + 1.f;
+		WorkingSize.X = AbscissaRange.Y + DragSize.X + 1.f;
+
 		++SlotCount;
-		UpdateLayoutOnDrag();
 	}
 	else
 	{
-		FMargin LastDropSlotOffset = DropSlots[SlotCount - 1]->OffsetAttr.Get(FMargin());
-		AbscissaRange.Y = LastDropSlotOffset.Left + LastDropSlotOffset.Right;
+		FMargin SlotOffset = DropSlots[SlotCount - 1]->OffsetAttr.Get(FMargin());
+		AbscissaRange.Y = SlotOffset.Left + SlotOffset.Right;
+
+		WorkingSize.X = DropSlots[SlotCount]->OffsetAttr.Get(FMargin()).Left + SDataprepGraphActionNode::DefaultWidth + InterSlotSpacing;
 
 		if(DragSlotOffset.Left > AbscissaRange.Y)
 		{
-			MouseOffset = DragSlotOffset.Left - AbscissaRange.Y;
+			OutScreenSpacePosition = GetAbsoluteSoftwareCursorPosition();
 
 			DragSlotOffset.Left = AbscissaRange.Y;
 			DragSlot->OffsetAttr.Set(DragSlotOffset);
 
-			HoveredDragIndex = GetHoveredActionNode(DragSlotOffset.Left, DragSlotOffset.Right);
+			IndexHovered = GetHoveredActionNode(DragSlotOffset.Left, DragSlotOffset.Right);
 		}
 
 		DropSlots[SlotCount]->AttachWidget(SNullWidget::NullWidget);
 		ActionSlots[SlotCount]->AttachWidget(SNullWidget::NullWidget);
 		TrackSlots[SlotCount]->AttachWidget(SNullWidget::NullWidget);
-
-		UpdateLayoutOnDrag();
 	}
 
-	Invalidate(EInvalidateWidgetReason::ChildOrder | EInvalidateWidgetReason::Layout);
-
-	return MouseOffset;
-}
+	UpdateLayoutOnDrag();
 
-bool SDataprepGraphTrackWidget::IsDraggedNodeOnEnds()
-{
-	if(DragSlot)
-	{
-		const FMargin DragSlotOffset = DragSlot->OffsetAttr.Get(FMargin());
-		return DragSlotOffset.Left == AbscissaRange.X || DragSlotOffset.Left == AbscissaRange.Y;
-	}
-
-	ensure(false);
-	return false;
+	Invalidate(EInvalidateWidgetReason::ChildOrder | EInvalidateWidgetReason::Layout);
 }
 
 FVector2D SDataprepGraphTrackWidget::GetNodeSize(const TSharedRef<SWidget>& Widget) const
diff --git a/Engine/Plugins/Enterprise/DataprepEditor/Source/DataprepEditor/Private/Widgets/DataprepGraph/SDataprepGraphTrackNode.h b/Engine/Plugins/Enterprise/DataprepEditor/Source/DataprepEditor/Private/Widgets/DataprepGraph/SDataprepGraphTrackNode.h
--- a/Engine/Plugins/Enterprise/DataprepEditor/Source/DataprepEditor/Private/Widgets/DataprepGraph/SDataprepGraphTrackNode.h
+++ b/Engine/Plugins/Enterprise/DataprepEditor/Source/DataprepEditor/Private/Widgets/DataprepGraph/SDataprepGraphTrackNode.h
@@ -98,12 +98,15 @@ public:
 
 	/**
 	 * Terminates the horizontal drag of an action node
-	 * @return Returns the new execution order for the dragged action node
+	 * @param bDoDrop Indicates if the drop must be done
 	 */
-	int32 OnEndNodeDrag();
+	void OnEndNodeDrag(bool bDoDrop);
 
-	/** Updates the position of other action nodes based on the position of the incoming node */
-	void OnNodeDragged( TSharedPtr<SDataprepGraphActionNode>& ActionNodePtr, const FVector2D& DragScreenSpacePosition, const FVector2D& ScreenSpaceDelta);
+	/**
+	 * Updates the position of other action nodes based on the position of the incoming node
+	 * @return True if the action node was actually dragged
+	 */
+	bool OnNodeDragged( TSharedPtr<SDataprepGraphActionNode>& ActionNodePtr, const FVector2D& DragScreenSpacePosition, const FVector2D& ScreenSpaceDelta);
 
 	/** Update the execution order of the actions and call ReArrangeActionNodes */
 	void OnActionsOrderChanged();
@@ -117,15 +120,20 @@ public:
 	/** Start editing of action asset associated to input EdGraphNode */
 	void RequestRename(const UEdGraphNode* Node);
 
+	/**
+	 * Change the view position of the panel based on input screen space position
+	 * @Note: Replaces SGrpahPanel::RequestDeferredPan which perform the change one frame too early
+	 */
 	void RequestViewportPan(const FVector2D& ScreenSpacePosition);
 
+	/** Helper to set the position of the graph nodes registered to the graph panel */
+	void UpdateProxyNode(TSharedRef<SDataprepGraphActionNode> ActioNodePtr, const FVector2D& ScreenSpacePosition);
+
 	/** Miscellaneous values used in the display */
 	static FVector2D TrackAnchor;
 
-protected:
-	// SWidget interface
-	virtual bool CustomPrepass(float LayoutScaleMultiplier) override;
-	// End of SWidget interface
+private:
+	FVector2D ComputePanAmount(const FVector2D& ScreenSpacePosition, float MaxPanSpeed = 200.f);
 
 private:
 	/** Pointer to the widget displaying the track */
@@ -143,14 +151,32 @@ private:
 	/** Indicates a drag is happening */
 	bool bNodeDragging;
 
+	/** Indicates a drag has started.  */
+	bool bNodeDragJustStarted;
+
+	/** Indicates a drag has started.  */
+	float LastDragDirection;
+
 	/**
-	 * Indicates to skip the next mouse position update as it has been triggered
-	 * by a call to FSlateApplication::SetCursorPos
+	 * Indicates to skip the next mouse position update during a drag
+	 * @remark Set to true when FSlateApplication::SetCursorPos has been called
 	 */
 	bool bSkipNextDragUpdate;
 
-	/** Cached of the last position of the cursor as the drag is happening */
-	FVector2D LastDragScreenSpacePosition;
+	/** Cached of the last position the software cursor was set to */
+	FVector2D LastSetCursorPosition;
+
+	/**
+	 * Indicates to skip the next refresh of the track widget layout
+	 * @remark Set to true when SNodePanel::RestoreViewSettings has been called
+	 */
+	bool bSkipRefreshLayout;
+
+	/** Indicates whether the mouse pointer left the graph panel on the left */
+	bool bCursorLeftOnLeft;
+
+	/** Indicates whether the mouse pointer left the graph panel on the right */
+	bool bCursorLeftOnRight;
 
 	/** Array of strong pointers to the UEdGraphNodes created for the Dataprep asset's actions */
 	TArray<TStrongObjectPtr<UDataprepGraphActionNode>> EdGraphActionNodes;
