Fixed build error on Linux - Was due to an ambiguous logic

#rnx

diff --git a/Engine/Plugins/Enterprise/DataprepEditor/Source/DataprepEditor/Private/Widgets/DataprepGraph/SDataprepGraphTrackNode.cpp b/Engine/Plugins/Enterprise/DataprepEditor/Source/DataprepEditor/Private/Widgets/DataprepGraph/SDataprepGraphTrackNode.cpp
--- a/Engine/Plugins/Enterprise/DataprepEditor/Source/DataprepEditor/Private/Widgets/DataprepGraph/SDataprepGraphTrackNode.cpp
+++ b/Engine/Plugins/Enterprise/DataprepEditor/Source/DataprepEditor/Private/Widgets/DataprepGraph/SDataprepGraphTrackNode.cpp
@@ -782,8 +782,8 @@ bool SDataprepGraphTrackNode::OnNodeDragged( TSharedPtr<SDataprepGraphActionNode
 	// When cursor was out of scope, do not process drag until cursor reaches marker
 	if(bCursorLeftOnLeft || bCursorLeftOnRight)
 	{
-		const bool bStillOutOfScope = bCursorLeftOnLeft && DragTrackPosition.X < SoftwareCursorPos.X ||
-									  bCursorLeftOnRight && DragTrackPosition.X > SoftwareCursorPos.X;
+		const bool bStillOutOfScope = (bCursorLeftOnLeft && DragTrackPosition.X < SoftwareCursorPos.X) ||
+									  (bCursorLeftOnRight && DragTrackPosition.X > SoftwareCursorPos.X);
 
 		// Still has not passed track marker, skip drag processing;
 		if(bStillOutOfScope)
