Fixed issue with assets' package path

#rnx

diff --git a/Engine/Plugins/Enterprise/DataprepEditor/Source/DataprepCore/Private/DataprepContentProducer.cpp b/Engine/Plugins/Enterprise/DataprepEditor/Source/DataprepCore/Private/DataprepContentProducer.cpp
--- a/Engine/Plugins/Enterprise/DataprepEditor/Source/DataprepCore/Private/DataprepContentProducer.cpp
+++ b/Engine/Plugins/Enterprise/DataprepEditor/Source/DataprepCore/Private/DataprepContentProducer.cpp
@@ -60,7 +60,7 @@ bool UDataprepContentProducer::Produce(const FDataprepProducerContext& InContext
 			// Verify each asset is within its own package
 			FString PackageName = Asset->GetOutermost()->GetName();
 			FString AssetName = Asset->GetName();
-			if(AssetName != PackageName)
+			if(AssetName != FPaths::GetBaseFilename(PackageName))
 			{
 				UPackage* NewPackage = NewObject<UPackage>( nullptr, *FPaths::Combine( PackageName, AssetName ), Asset->GetOutermost()->GetFlags() );
 				NewPackage->FullyLoad();
