Fixed crash on undo/redo on Dataprep asset with same number of actions but different content

#rnx

diff --git a/Engine/Plugins/Enterprise/DataprepEditor/Source/DataprepCore/Private/DataprepAsset.cpp b/Engine/Plugins/Enterprise/DataprepEditor/Source/DataprepCore/Private/DataprepAsset.cpp
--- a/Engine/Plugins/Enterprise/DataprepEditor/Source/DataprepCore/Private/DataprepAsset.cpp
+++ b/Engine/Plugins/Enterprise/DataprepEditor/Source/DataprepCore/Private/DataprepAsset.cpp
@@ -13,11 +13,13 @@
 #include "DataprepParameterizableObject.h"
 #include "Parameterization/DataprepParameterization.h"
 
+#include "Algo/Sort.h"
 #include "AssetRegistryModule.h"
 #include "BlueprintNodeBinder.h"
 #include "BlueprintNodeSpawner.h"
 #include "Kismet2/BlueprintEditorUtils.h"
 #include "Kismet2/KismetEditorUtilities.h"
+#include "Misc/SecureHash.h"
 #include "UObject/UObjectGlobals.h"
 
 #ifdef WITH_EDITOR
@@ -112,24 +114,33 @@ void UDataprepAsset::PreEditUndo()
 {
 	UDataprepAssetInterface::PreEditUndo();
 
-	// Cache number of actions to detect addition/removal
-	ActionCountBeforeUndoRedo = ActionAssets.Num();
+	// Cache signature of ordered array of actions before Undo/Redo
+	TArray<UDataprepActionAsset*> TempArray(ActionAssets);
+	Algo::Sort(TempArray);
+
+	SignatureBeforeUndoRedo = FMD5::HashBytes((uint8*)TempArray.GetData(), TempArray.Num() * sizeof(UDataprepActionAsset*));
 }
 
 void UDataprepAsset::PostEditUndo()
 {
 	UDataprepAssetInterface::PostEditUndo();
 
-	if(ActionCountBeforeUndoRedo != INDEX_NONE)
+	// Compute signature of ordered array of actions after Undo/Redo
+	TArray<UDataprepActionAsset*> TempArray(ActionAssets);
+	Algo::Sort(TempArray);
+
+	FString SignatureAfterUndoRedo = FMD5::HashBytes((uint8*)TempArray.GetData(), TempArray.Num() * sizeof(UDataprepActionAsset*));
+
+	if(!SignatureBeforeUndoRedo.IsEmpty())
 	{
-		OnActionChanged.Broadcast(nullptr, (ActionAssets.Num() == ActionCountBeforeUndoRedo) ? FDataprepAssetChangeType::ActionMoved : FDataprepAssetChangeType::ActionRemoved);
+		OnActionChanged.Broadcast(nullptr, (SignatureAfterUndoRedo == SignatureBeforeUndoRedo) ? FDataprepAssetChangeType::ActionMoved : FDataprepAssetChangeType::ActionRemoved);
 	}
 	else
 	{
 		ensure(false);
 	}
 
-	ActionCountBeforeUndoRedo = INDEX_NONE;
+	SignatureBeforeUndoRedo.Reset(0);
 }
 
 void UDataprepAsset::PostDuplicate(EDuplicateMode::Type DuplicateMode)
diff --git a/Engine/Plugins/Enterprise/DataprepEditor/Source/DataprepCore/Public/DataprepAsset.h b/Engine/Plugins/Enterprise/DataprepEditor/Source/DataprepCore/Public/DataprepAsset.h
--- a/Engine/Plugins/Enterprise/DataprepEditor/Source/DataprepCore/Public/DataprepAsset.h
+++ b/Engine/Plugins/Enterprise/DataprepEditor/Source/DataprepCore/Public/DataprepAsset.h
@@ -32,12 +32,6 @@ class DATAPREPCORE_API UDataprepAsset : public UDataprepAssetInterface
 	GENERATED_BODY()
 
 public:
-	UDataprepAsset()
-		: ActionCountBeforeUndoRedo(INDEX_NONE)
-	{
-	}
-
-
 	virtual ~UDataprepAsset() = default;
 
 	// UObject interface
@@ -243,6 +237,6 @@ private:
 
 	FOnStepObjectsAboutToBeRemoved OnStepObjectsAboutToBeRemoved;
 
-	/** Used to identify addition/removal of actions on undo/redo */
-	int32 ActionCountBeforeUndoRedo;
+	/** Used on undo/redo to detect addition/removal compared to change in execution order of actions */
+	FString SignatureBeforeUndoRedo;
 };
\ No newline at end of file
