Implement "Set Layer" operation

diff --git a/Engine/Plugins/Enterprise/DataprepEditor/Source/DataprepLibraries/Private/DataprepOperations.cpp b/Engine/Plugins/Enterprise/DataprepEditor/Source/DataprepLibraries/Private/DataprepOperations.cpp
--- a/Engine/Plugins/Enterprise/DataprepEditor/Source/DataprepLibraries/Private/DataprepOperations.cpp
+++ b/Engine/Plugins/Enterprise/DataprepEditor/Source/DataprepLibraries/Private/DataprepOperations.cpp
@@ -486,4 +486,14 @@ void FDataprepSetOutputFolderDetails::FolderName_TextChanged(const FText& Text)
 	bValidFolderName = ErrorMessage.IsEmpty();
 }
 
+void UDataprepAddToLayerOperation::OnExecution_Implementation(const FDataprepContext& InContext)
+{
+#ifdef LOG_TIME
+	DataprepOperationTime::FTimeLogger TimeLogger(TEXT("AddToLayer"), [&](FText Text) { this->LogInfo(Text); });
+#endif
+
+	// Execute operation
+	UDataprepOperationsLibrary::AddToLayer(InContext.Objects, LayerName);
+}
+
 #undef LOCTEXT_NAMESPACE
diff --git a/Engine/Plugins/Enterprise/DataprepEditor/Source/DataprepLibraries/Private/DataprepOperations.h b/Engine/Plugins/Enterprise/DataprepEditor/Source/DataprepLibraries/Private/DataprepOperations.h
--- a/Engine/Plugins/Enterprise/DataprepEditor/Source/DataprepLibraries/Private/DataprepOperations.h
+++ b/Engine/Plugins/Enterprise/DataprepEditor/Source/DataprepLibraries/Private/DataprepOperations.h
@@ -475,7 +475,6 @@ protected:
 	virtual void OnExecution_Implementation(const FDataprepContext& InContext) override;
 	//~ End UDataprepOperation Interface
 };
-
 // Customization of the details of the "Output to Folder" operation.
 class FDataprepSetOutputFolderDetails : public IDetailCustomization
 {
@@ -498,3 +497,25 @@ private:
 
 	TSharedPtr< class SEditableTextBox > TextBox;
 };
+
+UCLASS(Experimental, Category = ActorOperation, Meta = (DisplayName = "Add To Layer", ToolTip = "On each actor to process, add the actor to the layer"))
+class UDataprepAddToLayerOperation : public UDataprepOperation
+{
+	GENERATED_BODY()
+
+public:
+	// Name of the layer to add the actors to.
+	UPROPERTY(EditAnywhere, BlueprintReadWrite, Category = ActorOperation, Meta = (ToolTip = "Name of the layer to add the actors to"))
+	FName LayerName;
+
+	//~ Begin UDataprepOperation Interface
+public:
+	virtual FText GetCategory_Implementation() const override
+	{
+		return FDataprepOperationCategories::ActorOperation;
+	}
+
+protected:
+	virtual void OnExecution_Implementation(const FDataprepContext& InContext) override;
+	//~ End UDataprepOperation Interface
+};
diff --git a/Engine/Plugins/Enterprise/DataprepEditor/Source/DataprepLibraries/Private/DataprepOperationsLibrary.cpp b/Engine/Plugins/Enterprise/DataprepEditor/Source/DataprepLibraries/Private/DataprepOperationsLibrary.cpp
--- a/Engine/Plugins/Enterprise/DataprepEditor/Source/DataprepLibraries/Private/DataprepOperationsLibrary.cpp
+++ b/Engine/Plugins/Enterprise/DataprepEditor/Source/DataprepLibraries/Private/DataprepOperationsLibrary.cpp
@@ -962,4 +962,23 @@ void UDataprepOperationsLibrary::SetSubOuputFolder(const TArray<UObject*>& Selec
 	}
 }
 
+void UDataprepOperationsLibrary::AddToLayer(const TArray<UObject*>& SelectedObjects, const FName& LayerName)
+{
+	if (LayerName == NAME_None)
+	{
+		return;
+	}
+
+	for (UObject* Object : SelectedObjects)
+	{
+		if (AActor* Actor = Cast< AActor >(Object))
+		{
+			if (Actor && !Actor->IsPendingKill())
+			{
+				Actor->Layers.AddUnique(LayerName);
+			}
+		}
+	}
+}
+
 #undef LOCTEXT_NAMESPACE
diff --git a/Engine/Plugins/Enterprise/DataprepEditor/Source/DataprepLibraries/Private/DataprepOperationsLibrary.h b/Engine/Plugins/Enterprise/DataprepEditor/Source/DataprepLibraries/Private/DataprepOperationsLibrary.h
--- a/Engine/Plugins/Enterprise/DataprepEditor/Source/DataprepLibraries/Private/DataprepOperationsLibrary.h
+++ b/Engine/Plugins/Enterprise/DataprepEditor/Source/DataprepLibraries/Private/DataprepOperationsLibrary.h
@@ -337,6 +337,15 @@ public:
 	UFUNCTION(BlueprintCallable, Category = "Dataprep | Operation")
 	static void SetSubOuputFolder(const TArray<UObject*>& SelectedObjects, const FString& SubFolderName);
 
+	/**
+	 * Add all Actors to a given layer.
+	 * @param SelectedObjects:	Objects to consider
+	 * @param LayerName:	Name of the sub-folder
+	 * @note - This operation only applies on assets
+	 */
+	UFUNCTION(BlueprintCallable, Category = "Dataprep | Operation")
+	static void AddToLayer(const TArray<UObject*>& SelectedObjects, const FName& LayerName);
+
 private:
 	static void SubstituteMaterial(const TArray<UObject*>& SelectedObjects, const FString& MaterialSearch, EEditorScriptingStringMatchType StringMatch, const TArray<UMaterialInterface*>& MaterialList, UMaterialInterface* MaterialSubstitute);
 	static void SubstituteMesh(const TArray<UObject*>& SelectedObjects, const FString& MeshSearch, EEditorScriptingStringMatchType StringMatch, const TArray<UStaticMesh*>& MeshList, UStaticMesh* MeshSubstitute);
