- Worked around deprecated class errors in DatapropEditor enterprise plugin until it can be fixed for real

diff --git a/Engine/Plugins/Enterprise/DataprepEditor/Source/DataprepLibraries/Private/DataprepOperations.cpp b/Engine/Plugins/Enterprise/DataprepEditor/Source/DataprepLibraries/Private/DataprepOperations.cpp
--- a/Engine/Plugins/Enterprise/DataprepEditor/Source/DataprepLibraries/Private/DataprepOperations.cpp
+++ b/Engine/Plugins/Enterprise/DataprepEditor/Source/DataprepLibraries/Private/DataprepOperations.cpp
@@ -74,7 +74,7 @@ void UDataprepSetLODsOperation::OnExecution_Implementation(const FDataprepContex
 	}
 
 	// Fill up mesh reduction struct
-	FEditorScriptingMeshReductionOptions ReductionOptions;
+	FEditorScriptingMeshReductionOptions_Deprecated ReductionOptions;
 
 	ReductionOptions.bAutoComputeLODScreenSize = bAutoComputeLODScreenSize;
 	ReductionOptions.ReductionSettings.SetNum( LODCount );
diff --git a/Engine/Plugins/Enterprise/DataprepEditor/Source/DataprepLibraries/Private/DataprepOperations.h b/Engine/Plugins/Enterprise/DataprepEditor/Source/DataprepLibraries/Private/DataprepOperations.h
--- a/Engine/Plugins/Enterprise/DataprepEditor/Source/DataprepLibraries/Private/DataprepOperations.h
+++ b/Engine/Plugins/Enterprise/DataprepEditor/Source/DataprepLibraries/Private/DataprepOperations.h
@@ -103,14 +103,14 @@ class UDataprepSetSimpleCollisionOperation : public UDataprepOperation
 	GENERATED_BODY()
 
 	UDataprepSetSimpleCollisionOperation()
-		: ShapeType(EScriptingCollisionShapeType::Box)
+		: ShapeType(EScriptingCollisionShapeType_Deprecated::Box)
 	{
 	}
 
 public:
 	// Shape's of the collision geometry encompassing the static mesh
 	UPROPERTY(EditAnywhere, BlueprintReadWrite, Category = MeshOperation, meta = (ToolTip = "Shape's of the collision geometry encompassing the static mesh"))
-	EScriptingCollisionShapeType ShapeType;
+		EScriptingCollisionShapeType_Deprecated ShapeType;
 
 	//~ Begin UDataprepOperation Interface
 public:
diff --git a/Engine/Plugins/Enterprise/DataprepEditor/Source/DataprepLibraries/Private/DataprepOperationsLibrary.cpp b/Engine/Plugins/Enterprise/DataprepEditor/Source/DataprepLibraries/Private/DataprepOperationsLibrary.cpp
--- a/Engine/Plugins/Enterprise/DataprepEditor/Source/DataprepLibraries/Private/DataprepOperationsLibrary.cpp
+++ b/Engine/Plugins/Enterprise/DataprepEditor/Source/DataprepLibraries/Private/DataprepOperationsLibrary.cpp
@@ -378,7 +378,7 @@ namespace DataprepOperationsLibraryUtil
 	}
 } // ns DataprepOperationsLibraryUtil
 
-void UDataprepOperationsLibrary::SetLods(const TArray<UObject*>& SelectedObjects, const FEditorScriptingMeshReductionOptions& ReductionOptions, TArray<UObject*>& ModifiedObjects)
+void UDataprepOperationsLibrary::SetLods(const TArray<UObject*>& SelectedObjects, const FEditorScriptingMeshReductionOptions_Deprecated& ReductionOptions, TArray<UObject*>& ModifiedObjects)
 {
 	TSet<UStaticMesh*> SelectedMeshes = DataprepOperationsLibraryUtil::GetSelectedMeshes(SelectedObjects);
 
@@ -396,7 +396,7 @@ void UDataprepOperationsLibrary::SetLods(const TArray<UObject*>& SelectedObjects
 	}
 }
 
-void UDataprepOperationsLibrary::SetSimpleCollision(const TArray<UObject*>& SelectedObjects, const EScriptingCollisionShapeType ShapeType, TArray<UObject*>& ModifiedObjects)
+void UDataprepOperationsLibrary::SetSimpleCollision(const TArray<UObject*>& SelectedObjects, const EScriptingCollisionShapeType_Deprecated ShapeType, TArray<UObject*>& ModifiedObjects)
 {
 	TSet<UStaticMesh*> SelectedMeshes = DataprepOperationsLibraryUtil::GetSelectedMeshes(SelectedObjects);
 
@@ -404,11 +404,11 @@ void UDataprepOperationsLibrary::SetSimpleCollision(const TArray<UObject*>& Sele
 	bool bNeedRenderData = false;
 	switch (ShapeType)
 	{
-		case EScriptingCollisionShapeType::NDOP10_X:
-		case EScriptingCollisionShapeType::NDOP10_Y:
-		case EScriptingCollisionShapeType::NDOP10_Z:
-		case EScriptingCollisionShapeType::NDOP18:
-		case EScriptingCollisionShapeType::NDOP26:
+		case EScriptingCollisionShapeType_Deprecated::NDOP10_X:
+		case EScriptingCollisionShapeType_Deprecated::NDOP10_Y:
+		case EScriptingCollisionShapeType_Deprecated::NDOP10_Z:
+		case EScriptingCollisionShapeType_Deprecated::NDOP18:
+		case EScriptingCollisionShapeType_Deprecated::NDOP26:
 		{
 			bNeedRenderData = true;
 			break;
diff --git a/Engine/Plugins/Enterprise/DataprepEditor/Source/DataprepLibraries/Private/DataprepOperationsLibrary.h b/Engine/Plugins/Enterprise/DataprepEditor/Source/DataprepLibraries/Private/DataprepOperationsLibrary.h
--- a/Engine/Plugins/Enterprise/DataprepEditor/Source/DataprepLibraries/Private/DataprepOperationsLibrary.h
+++ b/Engine/Plugins/Enterprise/DataprepEditor/Source/DataprepLibraries/Private/DataprepOperationsLibrary.h
@@ -171,7 +171,7 @@ public:
 	 * Calls UEditorStaticMeshLibrary::SetLods on each static mesh of the resulting array.
 	 */
 	UFUNCTION(BlueprintCallable, Category = "Dataprep | Operation")
-	static void SetLods(const TArray<UObject*>& SelectedObjects, const FEditorScriptingMeshReductionOptions& ReductionOptions, TArray<UObject*>& ModifiedObjects);
+	static void SetLods(const TArray<UObject*>& SelectedObjects, const FEditorScriptingMeshReductionOptions_Deprecated& ReductionOptions, TArray<UObject*>& ModifiedObjects);
 
 	/**
 	 * Set one simple collision of the given shape type on the static meshes contained in the
@@ -187,7 +187,7 @@ public:
 	 * Calls UEditorStaticMeshLibrary::AddSimpleCollisions on each static mesh of the resulting array.
 	 */
 	UFUNCTION(BlueprintCallable, Category = "Dataprep | Operation")
-	static void SetSimpleCollision(const TArray<UObject*>& SelectedObjects, const EScriptingCollisionShapeType ShapeType, TArray<UObject*>& ModifiedObjects);
+	static void SetSimpleCollision(const TArray<UObject*>& SelectedObjects, const EScriptingCollisionShapeType_Deprecated ShapeType, TArray<UObject*>& ModifiedObjects);
 
 	/**
 	 * Add complex collision on the static meshes contained in the input array
