SceneOutliner refactor

Scene Outliner refactored to remove all actor-specific code.
Now uses modular modes and hierarchies to allow representation of any data model.

diff --git a/Engine/Config/BaseEditorPerProjectUserSettings.ini b/Engine/Config/BaseEditorPerProjectUserSettings.ini
--- a/Engine/Config/BaseEditorPerProjectUserSettings.ini
+++ b/Engine/Config/BaseEditorPerProjectUserSettings.ini
@@ -114,13 +114,16 @@ bDisplayFileSize=False
 bDisplayPaths=False
 bDisplayEditorOffset=False
 
-[/Script/SceneOutliner.SceneOutlinerSettings]
+[/Script/SceneOutliner.ActorBrowsingModeSettings]
 ; True when the Scene Outliner is hiding temporary/run-time Actors
 bHideTemporaryActors=False
 ; True when the Scene Outliner is showing only Actors that exist in the current level
 bShowOnlyActorsInCurrentLevel=False
 ; Scene Outliner Settings
 bShowOnlySelectedActors=False
+; Hide actor components
+bHideActorComponents=True
+
 
 [/Script/UnrealEd.ClassViewerSettings]
 ; Whether or not to display the hidden classes
diff --git a/Engine/Plugins/Enterprise/DataprepEditor/Source/DataprepEditor/Private/DataprepEditor.cpp b/Engine/Plugins/Enterprise/DataprepEditor/Source/DataprepEditor/Private/DataprepEditor.cpp
--- a/Engine/Plugins/Enterprise/DataprepEditor/Source/DataprepEditor/Private/DataprepEditor.cpp
+++ b/Engine/Plugins/Enterprise/DataprepEditor/Source/DataprepEditor/Private/DataprepEditor.cpp
@@ -41,7 +41,6 @@
 #include "Framework/MultiBox/MultiBoxBuilder.h"
 #include "GenericPlatform/GenericPlatformTime.h"
 #include "HAL/FileManager.h"
-#include "ICustomSceneOutliner.h"
 #include "Kismet2/BlueprintEditorUtils.h"
 #include "Kismet2/KismetEditorUtilities.h"
 #include "Materials/MaterialInstance.h"
@@ -1161,9 +1160,9 @@ void FDataprepEditor::RefreshColumnsForPreviewSystem()
 	{
 		SceneOutliner->RemoveColumn( FDataprepPreviewOutlinerColumn::ColumnID );
 		FSceneOutlinerModule& SceneOutlinerModule = FModuleManager::Get().LoadModuleChecked<FSceneOutlinerModule>("SceneOutliner");
-		SceneOutliner::FDefaultColumnInfo* ActorInfoColumPtr = SceneOutlinerModule.DefaultColumnMap.Find( SceneOutliner::FBuiltInColumnTypes::ActorInfo() );
+		SceneOutliner::FColumnInfo* ActorInfoColumPtr = SceneOutlinerModule.DefaultColumnMap.Find( SceneOutliner::FBuiltInColumnTypes::ActorInfo() );
 		check( ActorInfoColumPtr );
-		SceneOutliner->AddColumn( SceneOutliner::FBuiltInColumnTypes::ActorInfo(), ActorInfoColumPtr->ColumnInfo );
+		SceneOutliner->AddColumn( SceneOutliner::FBuiltInColumnTypes::ActorInfo(), *ActorInfoColumPtr );
 
 		AssetPreviewView->RemoveColumn( FDataprepPreviewAssetColumn::ColumnID );
 	}
diff --git a/Engine/Plugins/Enterprise/DataprepEditor/Source/DataprepEditor/Private/DataprepEditor.h b/Engine/Plugins/Enterprise/DataprepEditor/Source/DataprepEditor/Private/DataprepEditor.h
--- a/Engine/Plugins/Enterprise/DataprepEditor/Source/DataprepEditor/Private/DataprepEditor.h
+++ b/Engine/Plugins/Enterprise/DataprepEditor/Source/DataprepEditor/Private/DataprepEditor.h
@@ -123,6 +123,9 @@ public:
 	/** Return the number of steps the preview system is currently previewing */
 	int32 GetCountOfPreviewedSteps() const;
 
+	/** Handles change to selection in SceneOutliner */
+	void OnSceneOutlinerSelectionChanged(SceneOutliner::FTreeItemPtr ItemPtr, ESelectInfo::Type SelectionMode);
+
 private:
 	void BindCommands();
 	void OnSaveScene();
@@ -194,9 +197,6 @@ private:
 	/** Remove all temporary data remaining from previous runs of the Dataprep editor */
 	void CleanUpTemporaryDirectories();
 
-	/** Handles change to selection in SceneOutliner */
-	void OnSceneOutlinerSelectionChanged(SceneOutliner::FTreeItemPtr ItemPtr, ESelectInfo::Type SelectionMode);
-
 	bool OnCanExecuteNextStep(UDataprepActionAsset* ActionAsset);
 
 	/** Handles change to the content passed to an action */
@@ -233,7 +233,7 @@ private:
 	TSharedPtr<class SDataprepAssetView > DataprepAssetView;
 	TSharedPtr<SDataprepGraphEditor> GraphEditor;
 
-	TSharedPtr<class ICustomSceneOutliner> SceneOutliner;
+	TSharedPtr<class ISceneOutliner> SceneOutliner;
 
 	/** Command list for the pipeline editor */
 	TSharedPtr<FUICommandList> GraphEditorCommands;
diff --git a/Engine/Plugins/Enterprise/DataprepEditor/Source/DataprepEditor/Private/DataprepEditorOutliner.cpp b/Engine/Plugins/Enterprise/DataprepEditor/Source/DataprepEditor/Private/DataprepEditorOutliner.cpp
--- a/Engine/Plugins/Enterprise/DataprepEditor/Source/DataprepEditor/Private/DataprepEditorOutliner.cpp
+++ b/Engine/Plugins/Enterprise/DataprepEditor/Source/DataprepEditor/Private/DataprepEditorOutliner.cpp
@@ -2,9 +2,14 @@
 
 #include "DataprepEditor.h"
 
-#include "ICustomSceneOutliner.h"
+#include "ISceneOutliner.h"
 #include "ISceneOutlinerColumn.h"
 #include "SceneOutlinerModule.h"
+#include "DataprepEditorOutlinerMode.h"
+#include "ActorTreeItem.h"
+#include "ComponentTreeItem.h"
+#include "WorldTreeItem.h"
+#include "FolderTreeItem.h"
 
 #include "Modules/ModuleManager.h"
 #include "SceneOutlinerModule.h"
@@ -21,207 +26,64 @@ namespace DataprepEditorSceneOutlinerUtils
 	/** Specifics for the scene outliner */
 	using namespace SceneOutliner;
 
-	/**
-	 * This struct is used to force the scene outliner to refuse any rename request
-	 */
-	struct FCanRenameItem : public TTreeItemGetter<bool>
-	{
-		virtual bool Get(const SceneOutliner::FActorTreeItem& ActorItem) const override { return false; };
-		virtual bool Get(const SceneOutliner::FWorldTreeItem& WorldItem) const override { return false; }
-		virtual bool Get(const SceneOutliner::FFolderTreeItem& FolderItem) const override { return false; }
-		virtual bool Get(const SceneOutliner::FComponentTreeItem& ComponentItem) const override { return false; }
-		virtual bool Get(const SceneOutliner::FSubComponentTreeItem& SubComponentItem) const override { return false; }
-	};
-
-	/**
-	 * Use this struct to match the scene outliers selection to a dataprep editor selection
-	 */
-	struct FSynchroniseSelectionToSceneOutliner : public TTreeItemGetter<bool>
-	{
-		FSynchroniseSelectionToSceneOutliner(TSharedRef<FDataprepEditor> InDataprepEditor)
-			: DataprepEditorPtr(InDataprepEditor)
-		{
-		};
-
-		virtual bool Get(const FActorTreeItem& ActorItem) const override
-		{
-			if (const FDataprepEditor* DataprepEditor = DataprepEditorPtr.Pin().Get())
-			{
-				return DataprepEditor->GetWorldItemsSelection().Contains(ActorItem.Actor);
-			}
-			return false;
-		}
-
-		virtual bool Get(const FWorldTreeItem& WorldItem) const override
-		{
-			return false;
-		}
-		virtual bool Get(const FFolderTreeItem& FolderItem) const override
-		{
-			return false;
-		}
-		virtual bool Get(const FComponentTreeItem& ComponentItem) const override
-		{
-			if (const FDataprepEditor* DataprepEditor = DataprepEditorPtr.Pin().Get())
-			{
-				return DataprepEditor->GetWorldItemsSelection().Contains(ComponentItem.Component);
-			}
-			return false;
-		}
-		virtual bool Get(const FSubComponentTreeItem& SubComponentItem) const override
-		{
-			// return this for now as it seams that subcomponent Item is broken or doesn't do what want
-			return false;
-		}
-
-	private:
-		TWeakPtr<FDataprepEditor> DataprepEditorPtr;
-	};
-
-
 	/**
 	 * Use this struct to get the selection from the scene outliner
 	 */
-	struct FGetSelectionFromSceneOutliner : public ITreeItemVisitor
+	struct FGetSelectionFromSceneOutliner
 	{
 		mutable TSet<TWeakObjectPtr<UObject>> Selection;
 
-		virtual void Visit(const FActorTreeItem& ActorItem) const override
-		{
-			Selection.Add(ActorItem.Actor);
-		}
-
-		virtual void Visit(const FWorldTreeItem& WorldItem) const override {}
-		virtual void Visit(const FFolderTreeItem& FolderItem) const override {}
-		virtual void Visit(const FComponentTreeItem& ComponentItem) const override
-		{
-			this->Selection.Add(ComponentItem.Component);
-		}
-
-		virtual void Visit(const FSubComponentTreeItem& SubComponentItem) const override {}
-	};
-
-	struct FGetVisibilityVisitor : TTreeItemGetter<bool>
-	{
-		/** Map of tree item to visibility */
-		mutable TMap<const ITreeItem*, bool> VisibilityInfo;
-
-		bool RecurseChildren(const ITreeItem& Item) const
-		{
-			if (const bool* Info = VisibilityInfo.Find(&Item))
-			{
-				return *Info;
-			}
-			else
-			{
-				bool bIsVisible = false;
-				for (const auto& ChildPtr : Item.GetChildren())
-				{
-					auto Child = ChildPtr.Pin();
-					if (Child.IsValid() && Child->Get(*this))
-					{
-						bIsVisible = true;
-						break;
-					}
-				}
-				VisibilityInfo.Add(&Item, bIsVisible);
-
-				return bIsVisible;
-			}
-		}
-
-		bool Get(const FActorTreeItem& ActorItem) const
+		void TryAdd(const ITreeItem& Item) const
 		{
-			if (const bool* Info = VisibilityInfo.Find(&ActorItem))
+			if (const FActorTreeItem* ActorItem = Item.CastTo<FActorTreeItem>())
 			{
-				return *Info;
+				Selection.Add(ActorItem->Actor);
 			}
-			else
+			else if (const FComponentTreeItem* ComponentItem = Item.CastTo<FComponentTreeItem>())
 			{
-				const AActor* Actor = ActorItem.Actor.Get();
-
-				const bool bIsVisible = Actor && !Actor->IsTemporarilyHiddenInEditor(true);
-				VisibilityInfo.Add(&ActorItem, bIsVisible);
-
-				return bIsVisible;
+				Selection.Add(ComponentItem->Component);
 			}
 		}
-
-		bool Get(const FWorldTreeItem& WorldItem) const
-		{
-			return RecurseChildren(WorldItem);
-		}
-
-		bool Get(const FFolderTreeItem& FolderItem) const
-		{
-			return RecurseChildren(FolderItem);
-		}
 	};
 
-	struct FSetVisibilityVisitor : IMutableTreeItemVisitor
+	void SetVisibility(ITreeItem& Item, bool bIsVisible, TWeakPtr<SDataprepEditorViewport> Viewport)
 	{
-		/** Whether this item should be visible or not */
-		const bool bSetVisibility;
-		TWeakPtr<SDataprepEditorViewport> Viewport;
-
-		FSetVisibilityVisitor(bool bInSetVisibility, TWeakPtr<SDataprepEditorViewport> InViewport)
-			: bSetVisibility(bInSetVisibility)
-			, Viewport(InViewport)
-		{
-		}
-
-		virtual void Visit(FActorTreeItem& ActorItem) const override
+		if (FActorTreeItem* ActorItem = Item.CastTo<FActorTreeItem>())
 		{
-			AActor* Actor = ActorItem.Actor.Get();
+			AActor* Actor = ActorItem->Actor.Get();
 			if (Actor)
 			{
 				// Save the actor to the transaction buffer to support undo/redo, but do
 				// not call Modify, as we do not want to dirty the actor's package and
 				// we're only editing temporary, transient values
 				SaveToTransactionBuffer(Actor, false);
-				Actor->SetIsTemporarilyHiddenInEditor(!bSetVisibility);
+				Actor->SetIsTemporarilyHiddenInEditor(!bIsVisible);
 
-				Viewport.Pin()->SetActorVisibility(Actor, bSetVisibility);
+				Viewport.Pin()->SetActorVisibility(Actor, bIsVisible);
 
 				// Apply the same visibility to the actors children
-				for (auto& ChildPtr : ActorItem.GetChildren())
+				for (auto& ChildPtr : ActorItem->GetChildren())
 				{
 					auto Child = ChildPtr.Pin();
 					if (Child.IsValid())
 					{
-						FSetVisibilityVisitor Visibility(bSetVisibility, Viewport);
-						Child->Visit(Visibility);
+						SetVisibility(*Child, bIsVisible, Viewport);
 					}
 				}
 			}
 		}
-
-		virtual void Visit(FWorldTreeItem& WorldItem) const override
+		else if (Item.IsA<FWorldTreeItem>() || Item.IsA<FFolderTreeItem>())
 		{
-			for (auto& ChildPtr : WorldItem.GetChildren())
+			for (auto& ChildPtr : Item.GetChildren())
 			{
 				auto Child = ChildPtr.Pin();
 				if (Child.IsValid())
 				{
-					FSetVisibilityVisitor Visibility(bSetVisibility, Viewport);
-					Child->Visit(Visibility);
+					SetVisibility(*Child, bIsVisible, Viewport);
 				}
 			}
 		}
-
-		virtual void Visit(FFolderTreeItem& FolderItem) const override
-		{
-			for (auto& ChildPtr : FolderItem.GetChildren())
-			{
-				auto Child = ChildPtr.Pin();
-				if (Child.IsValid())
-				{
-					FSetVisibilityVisitor Visibility(bSetVisibility, Viewport);
-					Child->Visit(Visibility);
-				}
-			}
-		}
-	};
+	}
 
 	class FVisibilityDragDropOp : public FDragDropOperation, public TSharedFromThis<FVisibilityDragDropOp>
 	{
@@ -288,7 +150,7 @@ namespace DataprepEditorSceneOutlinerUtils
 		virtual void SortItems(TArray<FTreeItemPtr>& RootItems, const EColumnSortMode::Type SortMode) const override {}
 
 		/** Check whether the specified item is visible */
-		bool IsItemVisible(const ITreeItem& Item) { return Item.Get(VisibilityCache); }
+		bool IsItemVisible(const ITreeItem& Item) { return VisibilityCache.GetVisibility(Item); }
 
 		TWeakPtr<SDataprepEditorViewport> GetViewport() const
 		{
@@ -301,8 +163,8 @@ namespace DataprepEditorSceneOutlinerUtils
 
 		TWeakPtr<SDataprepEditorViewport> WeakSceneViewport;
 
-		/** Visitor used to get (and cache) visibilty for items. Cahced per-frame to avoid expensive recursion. */
-		FGetVisibilityVisitor VisibilityCache;
+		/** Get and cache visibility for items. Cached per-frame to avoid expensive recursion. */
+		FGetVisibilityCache VisibilityCache;
 	};
 
 	/** Widget responsible for managing the visibility for a single actor */
@@ -371,13 +233,13 @@ namespace DataprepEditorSceneOutlinerUtils
 			// We operate on all the selected items if the specified item is selected
 			if (Tree.IsItemSelected(TreeItem.ToSharedRef()))
 			{
-				const FSetVisibilityVisitor Visitor(bVisible, WeakColumn.Pin()->GetViewport());
+				auto Viewport = WeakColumn.Pin()->GetViewport();
 
 				for (auto& SelectedItem : Tree.GetSelectedItems())
 				{
 					if (IsVisible(SelectedItem, Column) != bVisible)
 					{
-						SelectedItem->Visit(Visitor);
+						DataprepEditorSceneOutlinerUtils::SetVisibility(*SelectedItem, bVisible, Viewport);
 					}
 				}
 
@@ -464,8 +326,7 @@ namespace DataprepEditorSceneOutlinerUtils
 
 			if (TreeItem.IsValid() && Outliner.IsValid() && IsVisible() != bVisible)
 			{
-				FSetVisibilityVisitor Visitor(bVisible, WeakColumn.Pin()->GetViewport());
-				TreeItem->Visit(Visitor);
+				DataprepEditorSceneOutlinerUtils::SetVisibility(*TreeItem, bVisible, WeakColumn.Pin()->GetViewport());
 
 				Outliner->Refresh();
 
@@ -503,11 +364,16 @@ namespace DataprepEditorSceneOutlinerUtils
 void FDataprepEditor::CreateScenePreviewTab()
 {
 	FSceneOutlinerModule& SceneOutlinerModule = FModuleManager::Get().LoadModuleChecked<FSceneOutlinerModule>("SceneOutliner");
+	
+	SceneOutliner::FCreateOutlinerMode ModeFactory = SceneOutliner::FCreateOutlinerMode::CreateLambda([SpecifiedWorld = PreviewWorld.Get(), DataprepEditor = StaticCastSharedRef<FDataprepEditor>(AsShared())](SceneOutliner::SSceneOutliner* Outliner)
+		{
+			return new FDataprepEditorOutlinerMode(Outliner, DataprepEditor, SpecifiedWorld);
+		});
 
 	SceneOutliner::FInitializationOptions SceneOutlinerOptions;
-	SceneOutlinerOptions.SpecifiedWorldToDisplay = PreviewWorld.Get();
+	SceneOutlinerOptions.ModeFactory = ModeFactory;
 
-	SceneOutliner = SceneOutlinerModule.CreateCustomSceneOutliner(SceneOutlinerOptions);
+	SceneOutliner = SceneOutlinerModule.CreateSceneOutliner(SceneOutlinerOptions);
 
 	// Add our custom visibility gutter
 
@@ -526,26 +392,13 @@ void FDataprepEditor::CreateScenePreviewTab()
 
 	SceneOutliner->AddColumn(DataprepEditorSceneOutlinerUtils::FPreviewSceneOutlinerGutter::GetID(), ColumnInfo);
 
-	// Add the default outliner columns
+	// Add the default columns
 	const SceneOutliner::FSharedOutlinerData& SharedData = SceneOutliner->GetSharedData();
 	for (auto& DefaultColumn : SceneOutlinerModule.DefaultColumnMap)
 	{
-		if (!DefaultColumn.Value.ValidMode.IsSet() || SharedData.Mode == DefaultColumn.Value.ValidMode.GetValue())
-		{
-			SceneOutliner->AddColumn(DefaultColumn.Key, DefaultColumn.Value.ColumnInfo);
-		}
+		SceneOutliner->AddColumn(DefaultColumn.Key, DefaultColumn.Value);
 	}
 
-	SceneOutliner->SetSelectionMode(ESelectionMode::Multi)
-		.SetCanRenameItem(MakeUnique<DataprepEditorSceneOutlinerUtils::FCanRenameItem>())
-		.SetShouldSelectItemWhenAdded(MakeUnique<DataprepEditorSceneOutlinerUtils::FSynchroniseSelectionToSceneOutliner>(StaticCastSharedRef<FDataprepEditor>(AsShared())))
-		.SetShowActorComponents(false)
-		.SetShownOnlySelected(false)
-		.SetShowOnlyCurrentLevel(false)
-		.SetHideTemporaryActors(false);
-
-	SceneOutliner->GetOnItemSelectionChanged().AddSP(this, &FDataprepEditor::OnSceneOutlinerSelectionChanged);
-
 	SAssignNew(ScenePreviewView, SBorder)
 		.Padding(2.f)
 		.BorderImage(FEditorStyle::GetBrush("ToolPanel.GroupBorder"))
@@ -564,14 +417,14 @@ void FDataprepEditor::OnSceneOutlinerSelectionChanged(SceneOutliner::FTreeItemPt
 {
 	using namespace SceneOutliner;
 
-	DataprepEditorSceneOutlinerUtils::FGetSelectionFromSceneOutliner Visitor;
+	DataprepEditorSceneOutlinerUtils::FGetSelectionFromSceneOutliner Selector;
 
 	for (FTreeItemPtr Item : SceneOutliner->GetTree().GetSelectedItems())
 	{
-		Item->Visit(Visitor);
+		Selector.TryAdd(*Item);
 	}
 
-	SetWorldObjectsSelection(MoveTemp(Visitor.Selection), EWorldSelectionFrom::SceneOutliner);
+	SetWorldObjectsSelection(MoveTemp(Selector.Selection), EWorldSelectionFrom::SceneOutliner);
 }
 
 void FDataprepEditor::SetWorldObjectsSelection(TSet<TWeakObjectPtr<UObject>>&& NewSelection, EWorldSelectionFrom SelectionFrom /* = EWorldSelectionFrom::Unknow */)
diff --git a/Engine/Plugins/Enterprise/DataprepEditor/Source/DataprepEditor/Private/DataprepEditorOutlinerMode.h b/Engine/Plugins/Enterprise/DataprepEditor/Source/DataprepEditor/Private/DataprepEditorOutlinerMode.h
new file mode 100644
--- /dev/null
+++ b/Engine/Plugins/Enterprise/DataprepEditor/Source/DataprepEditor/Private/DataprepEditorOutlinerMode.h
@@ -0,0 +1,68 @@
+// Copyright Epic Games, Inc. All Rights Reserved.
+
+#pragma once
+
+#include "ActorMode.h"
+#include "ActorTreeItem.h"
+#include "ComponentTreeItem.h"
+
+namespace DataprepEditorSceneOutlinerUtils
+{
+	using namespace SceneOutliner;
+
+	/**
+	* Use this struct to match the scene outliers selection to a dataprep editor selection
+	*/
+	struct FSynchroniseSelectionToSceneOutliner
+	{
+		FSynchroniseSelectionToSceneOutliner(TWeakPtr<FDataprepEditor> InDataprepEditor)
+			: DataprepEditorPtr(InDataprepEditor)
+		{
+		};
+
+		bool operator()(const ITreeItem& Item) const
+		{
+			if (const FActorTreeItem* ActorItem = Item.CastTo<FActorTreeItem>())
+			{
+				if (const FDataprepEditor* DataprepEditor = DataprepEditorPtr.Pin().Get())
+				{
+					return DataprepEditor->GetWorldItemsSelection().Contains(ActorItem->Actor);
+				}
+			}
+			else if (const FComponentTreeItem* ComponentItem = Item.CastTo<FComponentTreeItem>())
+			{
+				if (const FDataprepEditor* DataprepEditor = DataprepEditorPtr.Pin().Get())
+				{
+					return DataprepEditor->GetWorldItemsSelection().Contains(ComponentItem->Component);
+				}
+			}
+
+			return false;
+		}
+	private:
+		TWeakPtr<FDataprepEditor> DataprepEditorPtr;
+	};
+}
+
+class FDataprepEditorOutlinerMode : public SceneOutliner::FActorMode
+{
+public:
+	FDataprepEditorOutlinerMode(SceneOutliner::SSceneOutliner* InSceneOutliner, TWeakPtr<FDataprepEditor> InDataprepEditor, TWeakObjectPtr<UWorld> InSpecifiedWorldToDisplay = nullptr)
+		: SceneOutliner::FActorMode(InSceneOutliner, true, InSpecifiedWorldToDisplay)
+		, DataprepEditorPtr(InDataprepEditor)
+	{}
+
+	virtual bool CanRenameItem(const SceneOutliner::ITreeItem& Item) const override { return false; }
+	virtual ESelectionMode::Type GetSelectionMode() const override { return ESelectionMode::Multi; }
+
+	virtual void OnItemSelectionChanged(SceneOutliner::FTreeItemPtr Item, ESelectInfo::Type SelectionType, const SceneOutliner::FItemSelection& Selection) override
+	{
+		auto DataprepEditor = DataprepEditorPtr.Pin();
+		if (DataprepEditor)
+		{
+			DataprepEditor->OnSceneOutlinerSelectionChanged(Item, SelectionType);
+		}
+	}
+private:
+	TWeakPtr<FDataprepEditor> DataprepEditorPtr;
+};
\ No newline at end of file
diff --git a/Engine/Plugins/Enterprise/DataprepEditor/Source/DataprepEditor/Private/PreviewSystem/DataprepPreviewSceneOutlinerColumn.cpp b/Engine/Plugins/Enterprise/DataprepEditor/Source/DataprepEditor/Private/PreviewSystem/DataprepPreviewSceneOutlinerColumn.cpp
--- a/Engine/Plugins/Enterprise/DataprepEditor/Source/DataprepEditor/Private/PreviewSystem/DataprepPreviewSceneOutlinerColumn.cpp
+++ b/Engine/Plugins/Enterprise/DataprepEditor/Source/DataprepEditor/Private/PreviewSystem/DataprepPreviewSceneOutlinerColumn.cpp
@@ -8,9 +8,6 @@
 #include "ActorTreeItem.h"
 #include "ComponentTreeItem.h"
 #include "ISceneOutliner.h"
-#include "SceneOutlinerVisitorTypes.h"
-#include "SceneOutlinerVisitorTypes.h"
-#include "SubComponentTreeItem.h"
 #include "UObject/Object.h"
 #include "UObject/WeakObjectPtrTemplates.h"
 #include "Widgets/SNullWidget.h"
@@ -20,14 +17,19 @@
 
 namespace DataprepPreviewOutlinerColumnUtils
 {
-	struct FObjectGetter : public SceneOutliner::TTreeItemGetter<TWeakObjectPtr<UObject>>
+	
+	UObject* GetObjectPtr(const SceneOutliner::ITreeItem& Item)
 	{
-		virtual TWeakObjectPtr<UObject> Get(const SceneOutliner::FActorTreeItem& ActorItem) const override { return ActorItem.Actor; };
-		virtual TWeakObjectPtr<UObject> Get(const SceneOutliner::FWorldTreeItem& WorldItem) const override { return nullptr; }
-		virtual TWeakObjectPtr<UObject> Get(const SceneOutliner::FFolderTreeItem& FolderItem) const override { return nullptr; }
-		virtual TWeakObjectPtr<UObject> Get(const SceneOutliner::FComponentTreeItem& ComponentItem) const override { return ComponentItem.Component; }
-		virtual TWeakObjectPtr<UObject> Get(const SceneOutliner::FSubComponentTreeItem& SubComponentItem) const override { return SubComponentItem.ParentComponent; }
-	};
+		if (const SceneOutliner::FActorTreeItem* ActorItem = Item.CastTo<SceneOutliner::FActorTreeItem>())
+		{
+			return ActorItem->Actor.Get();
+		}
+		else if (const SceneOutliner::FComponentTreeItem* ComponentItem = Item.CastTo<SceneOutliner::FComponentTreeItem>())
+		{
+			return ComponentItem->Component.Get();
+		}
+		return nullptr;
+	}
 }
 
 const FName FDataprepPreviewOutlinerColumn::ColumnID = FName( TEXT("DataprepPreview") );
@@ -56,10 +58,7 @@ SHeaderRow::FColumn::FArguments FDataprepPreviewOutlinerColumn::ConstructHeaderR
 
 const TSharedRef<SWidget> FDataprepPreviewOutlinerColumn::ConstructRowWidget(SceneOutliner::FTreeItemRef TreeItem, const STableRow<SceneOutliner::FTreeItemPtr>& Row)
 {
-	DataprepPreviewOutlinerColumnUtils::FObjectGetter Visitor; 
-	TreeItem->Visit( Visitor );
-
-	if ( UObject* Object = Visitor.Result().Get() )
+	if ( UObject* Object = DataprepPreviewOutlinerColumnUtils::GetObjectPtr(*TreeItem))
 	{
 		if ( TSharedPtr<ISceneOutliner> SceneOutliner = WeakSceneOutliner.Pin() )
 		{ 
@@ -73,9 +72,7 @@ const TSharedRef<SWidget> FDataprepPreviewOutlinerColumn::ConstructRowWidget(Sce
 
 void FDataprepPreviewOutlinerColumn::PopulateSearchStrings(const SceneOutliner::ITreeItem& Item, TArray< FString >& OutSearchStrings) const
 {
-	DataprepPreviewOutlinerColumnUtils::FObjectGetter Visitor;
-	Item.Visit(Visitor);
-	if (UObject* Object = Visitor.Result().Get())
+	if (UObject* Object = DataprepPreviewOutlinerColumnUtils::GetObjectPtr(Item))
 	{
 		if ( TSharedPtr<FDataprepPreviewProcessingResult> PreviewResult = CachedPreviewData->GetPreviewDataForObject( Object ) )
 		{
@@ -88,11 +85,8 @@ void FDataprepPreviewOutlinerColumn::SortItems(TArray<SceneOutliner::FTreeItemPt
 {
 	OutItems.Sort([this, SortMode](const SceneOutliner::FTreeItemPtr& First, const SceneOutliner::FTreeItemPtr& Second)
 	{
-			DataprepPreviewOutlinerColumnUtils::FObjectGetter Visitor;
-			First->Visit( Visitor );
-			UObject* FirstObject = Visitor.Result().Get();
-			Second->Visit( Visitor );
-			UObject* SecondObject = Visitor.Result().Get();
+			UObject* FirstObject = DataprepPreviewOutlinerColumnUtils::GetObjectPtr(*First);
+			UObject* SecondObject = DataprepPreviewOutlinerColumnUtils::GetObjectPtr(*Second);
 			if ( FirstObject && SecondObject )
 			{
 
@@ -133,8 +127,9 @@ void FDataprepPreviewOutlinerColumn::SortItems(TArray<SceneOutliner::FTreeItemPt
 			}
 
 		// If all else fail filter by name (always Ascending)
-		int32 SortPriorityFirst = First->GetTypeSortPriority();
-		int32 SortPrioritySecond = Second->GetTypeSortPriority();
+		auto Outliner = WeakSceneOutliner.Pin();
+		int32 SortPriorityFirst = Outliner->GetTypeSortPriority(*First);
+		int32 SortPrioritySecond = Outliner->GetTypeSortPriority(*Second);
 		if ( SortPriorityFirst != SortPrioritySecond )
 		{
 			return SortPriorityFirst < SortPrioritySecond;
diff --git a/Engine/Plugins/Enterprise/VariantManager/Source/VariantManager/Private/DisplayNodes/VariantManagerActorNode.cpp b/Engine/Plugins/Enterprise/VariantManager/Source/VariantManager/Private/DisplayNodes/VariantManagerActorNode.cpp
--- a/Engine/Plugins/Enterprise/VariantManager/Source/VariantManager/Private/DisplayNodes/VariantManagerActorNode.cpp
+++ b/Engine/Plugins/Enterprise/VariantManager/Source/VariantManager/Private/DisplayNodes/VariantManagerActorNode.cpp
@@ -22,6 +22,7 @@
 #include "Modules/ModuleManager.h"
 #include "SceneOutlinerModule.h"
 #include "SceneOutlinerPublicTypes.h"
+#include "ActorTreeItem.h"
 #include "ScopedTransaction.h"
 #include "Styling/SlateIconFinder.h"
 #include "Textures/SlateIcon.h"
@@ -507,8 +508,6 @@ void FVariantManagerActorNode::AddAssignActorSubMenu(FMenuBuilder& MenuBuilder)
 	// Set up a menu entry to assign an actor to the object binding node
 	FInitializationOptions InitOptions;
 	{
-		InitOptions.Mode = ESceneOutlinerMode::ActorPicker;
-
 		// We hide the header row to keep the UI compact.
 		InitOptions.bShowHeaderRow = false;
 		InitOptions.bShowSearchBox = true;
@@ -518,7 +517,7 @@ void FVariantManagerActorNode::AddAssignActorSubMenu(FMenuBuilder& MenuBuilder)
 		InitOptions.ColumnMap.Add(FBuiltInColumnTypes::Label(), FColumnInfo(EColumnVisibility::Visible, 0));
 
 		// Only display actors that are not possessed already
-		InitOptions.Filters->AddFilterPredicate( FActorFilterPredicate::CreateLambda( IsActorValidForAssignment ) );
+		InitOptions.Filters->AddFilterPredicate<FActorTreeItem>(FActorTreeItem::FFilterPredicate::CreateLambda( IsActorValidForAssignment ));
 	}
 
 	// actor selector to allow the user to choose an actor
@@ -528,7 +527,7 @@ void FVariantManagerActorNode::AddAssignActorSubMenu(FMenuBuilder& MenuBuilder)
 		.MaxDesiredHeight(400.0f)
 		.WidthOverride(300.0f)
 		[
-			SceneOutlinerModule.CreateSceneOutliner(
+			SceneOutlinerModule.CreateActorPicker(
 				InitOptions,
 				FOnActorPicked::CreateLambda([&](AActor* Actor)
 				{
diff --git a/Engine/Plugins/Enterprise/VariantManager/Source/VariantManager/Private/SVariantManager.cpp b/Engine/Plugins/Enterprise/VariantManager/Source/VariantManager/Private/SVariantManager.cpp
--- a/Engine/Plugins/Enterprise/VariantManager/Source/VariantManager/Private/SVariantManager.cpp
+++ b/Engine/Plugins/Enterprise/VariantManager/Source/VariantManager/Private/SVariantManager.cpp
@@ -60,6 +60,7 @@
 #include "ObjectTools.h"
 #include "SceneOutlinerModule.h"
 #include "SceneOutlinerPublicTypes.h"
+#include "ActorTreeItem.h"
 #include "ScopedTransaction.h"
 #include "Widgets/Input/SButton.h"
 #include "Widgets/Input/SButton.h"
@@ -2650,13 +2651,12 @@ FReply SVariantManager::OnSummonAddActorMenu()
 	};
 
 	FInitializationOptions InitOptions;
-	InitOptions.Mode = ESceneOutlinerMode::ActorPicker;
 	InitOptions.bShowHeaderRow = true;
 	InitOptions.bShowSearchBox = true;
 	InitOptions.bShowCreateNewFolder = false;
 	InitOptions.bFocusSearchBoxWhenOpened = true;
 	InitOptions.ColumnMap.Add(FBuiltInColumnTypes::Label(), FColumnInfo(EColumnVisibility::Visible, 0));
-	InitOptions.Filters->AddFilterPredicate( FActorFilterPredicate::CreateLambda( IsActorValidForAssignment ) );
+	InitOptions.Filters->AddFilterPredicate<FActorTreeItem>(FActorTreeItem::FFilterPredicate::CreateLambda( IsActorValidForAssignment ));
 
 	// Create mini scene outliner menu
 	FSceneOutlinerModule& SceneOutlinerModule = FModuleManager::LoadModuleChecked<FSceneOutlinerModule>("SceneOutliner");
@@ -2665,7 +2665,7 @@ FReply SVariantManager::OnSummonAddActorMenu()
 		.MaxDesiredHeight(400.0f)
 		.WidthOverride(300.0f)
 		[
-			SceneOutlinerModule.CreateSceneOutliner(
+			SceneOutlinerModule.CreateActorPicker(
 				InitOptions,
 				FOnActorPicked::CreateLambda([SelectedVariants, VarMan, this](AActor* Actor)
 				{
diff --git a/Engine/Plugins/MovieScene/LevelSequenceEditor/Source/LevelSequenceEditor/Private/Misc/LevelSequenceEditorActorBinding.cpp b/Engine/Plugins/MovieScene/LevelSequenceEditor/Source/LevelSequenceEditor/Private/Misc/LevelSequenceEditorActorBinding.cpp
--- a/Engine/Plugins/MovieScene/LevelSequenceEditor/Source/LevelSequenceEditor/Private/Misc/LevelSequenceEditorActorBinding.cpp
+++ b/Engine/Plugins/MovieScene/LevelSequenceEditor/Source/LevelSequenceEditor/Private/Misc/LevelSequenceEditorActorBinding.cpp
@@ -7,6 +7,7 @@
 #include "Styling/SlateIconFinder.h"
 #include "SceneOutlinerPublicTypes.h"
 #include "SceneOutlinerModule.h"
+#include "ActorTreeItem.h"
 #include "Framework/MultiBox/MultiBoxBuilder.h"
 #include "Engine/Engine.h"
 #include "Framework/Application/SlateApplication.h"
@@ -101,8 +102,6 @@ void FLevelSequenceEditorActorBinding::AddPossessActorMenuExtensions(FMenuBuilde
 	// Set up a menu entry to add any arbitrary actor to the sequencer
 	FInitializationOptions InitOptions;
 	{
-		InitOptions.Mode = ESceneOutlinerMode::ActorPicker;
-
 		// We hide the header row to keep the UI compact.
 		InitOptions.bShowHeaderRow = false;
 		InitOptions.bShowSearchBox = true;
@@ -113,7 +112,7 @@ void FLevelSequenceEditorActorBinding::AddPossessActorMenuExtensions(FMenuBuilde
 		InitOptions.ColumnMap.Add(FBuiltInColumnTypes::Label(), FColumnInfo(EColumnVisibility::Visible, 0));
 
 		// Only display actors that are not possessed already
-		InitOptions.Filters->AddFilterPredicate(FActorFilterPredicate::CreateLambda(IsActorValidForPossession, ExistingPossessedObjects));
+		InitOptions.Filters->AddFilterPredicate<FActorTreeItem>(FActorTreeItem::FFilterPredicate::CreateLambda(IsActorValidForPossession, ExistingPossessedObjects));
 	}
 
 	// actor selector to allow the user to choose an actor
@@ -123,7 +122,7 @@ void FLevelSequenceEditorActorBinding::AddPossessActorMenuExtensions(FMenuBuilde
 		.MaxDesiredHeight(400.0f)
 		.WidthOverride(300.0f)
 		[
-			SceneOutlinerModule.CreateSceneOutliner(
+			SceneOutlinerModule.CreateActorPicker(
 				InitOptions,
 				FOnActorPicked::CreateLambda([=](AActor* Actor){
 					// Create a new binding for this actor
diff --git a/Engine/Plugins/VirtualProduction/Takes/Source/TakeRecorderSources/Private/TakeRecorderSourcesModule.cpp b/Engine/Plugins/VirtualProduction/Takes/Source/TakeRecorderSources/Private/TakeRecorderSourcesModule.cpp
--- a/Engine/Plugins/VirtualProduction/Takes/Source/TakeRecorderSources/Private/TakeRecorderSourcesModule.cpp
+++ b/Engine/Plugins/VirtualProduction/Takes/Source/TakeRecorderSources/Private/TakeRecorderSourcesModule.cpp
@@ -32,6 +32,7 @@
 #include "SceneOutlinerModule.h"
 #include "SceneOutlinerPublicTypes.h"
 #include "Subsystems/AssetEditorSubsystem.h"
+#include "ActorTreeItem.h"
 
 #include "TakeRecorderMicrophoneAudioSource.h"
 #include "TakeRecorderWorldSource.h"
@@ -206,8 +207,8 @@ struct FActorTakeRecorderDropHandler : ITakeRecorderDropHandler
 		else if (OperationPtr->IsOfType<FSceneOutlinerDragDropOp>())
 		{
 			FSceneOutlinerDragDropOp* OutlinerOp = static_cast<FSceneOutlinerDragDropOp*>(OperationPtr);
-			FolderDrag = OutlinerOp->FolderOp.Get();
-			ActorDrag  = OutlinerOp->ActorOp.Get();
+			FolderDrag = OutlinerOp->GetSubOp<FFolderDragDropOp>().Get();
+			ActorDrag  = OutlinerOp->GetSubOp<FActorDragDropOp>().Get();
 
 		}
 		else if (OperationPtr->IsOfType<FActorDragDropOp>())
@@ -398,8 +399,6 @@ public:
 			// Set up a menu entry to add any arbitrary actor to the sequencer
 			FInitializationOptions InitOptions;
 			{
-				InitOptions.Mode = ESceneOutlinerMode::ActorPicker;
-
 				// We hide the header row to keep the UI compact.
 				InitOptions.bShowHeaderRow = false;
 				InitOptions.bShowSearchBox = true;
@@ -410,7 +409,7 @@ public:
 				InitOptions.ColumnMap.Add(FBuiltInColumnTypes::Label(), FColumnInfo(EColumnVisibility::Visible, 0));
 
 				// Only display actors that are not possessed already
-				InitOptions.Filters->AddFilterPredicate(FActorFilterPredicate::CreateLambda(OutlinerFilterPredicate));
+				InitOptions.Filters->AddFilterPredicate<FActorTreeItem>(FActorTreeItem::FFilterPredicate::CreateLambda(OutlinerFilterPredicate));
 			}
 
 			// actor selector to allow the user to choose an actor
@@ -420,7 +419,7 @@ public:
 				.MaxDesiredHeight(400.0f)
 				.WidthOverride(300.0f)
 				[
-					SceneOutlinerModule.CreateSceneOutliner(
+					SceneOutlinerModule.CreateActorPicker(
 						InitOptions,
 						FOnActorPicked::CreateLambda([Sources](AActor* Actor){
 							// Create a new binding for this actor
diff --git a/Engine/Source/Editor/Layers/Private/ActorsAssignedToSpecificLayersFilter.h b/Engine/Source/Editor/Layers/Private/ActorsAssignedToSpecificLayersFilter.h
--- a/Engine/Source/Editor/Layers/Private/ActorsAssignedToSpecificLayersFilter.h
+++ b/Engine/Source/Editor/Layers/Private/ActorsAssignedToSpecificLayersFilter.h
@@ -3,14 +3,15 @@
 
 #include "CoreMinimal.h"
 #include "SceneOutlinerFilters.h"
+#include "ActorTreeItem.h"
 
-class FActorsAssignedToSpecificLayersFilter : public SceneOutliner::FOutlinerFilter, public TSharedFromThis< FActorsAssignedToSpecificLayersFilter >
+class FActorsAssignedToSpecificLayersFilter : public SceneOutliner::TOutlinerFilter<SceneOutliner::FActorTreeItem>, public TSharedFromThis< FActorsAssignedToSpecificLayersFilter >
 {
 public:
 
 	/** By default, any types not handled by this filter will fail */
 	FActorsAssignedToSpecificLayersFilter()
-		: FOutlinerFilter(SceneOutliner::EDefaultFilterBehaviour::Fail)
+		: TOutlinerFilter<SceneOutliner::FActorTreeItem>(SceneOutliner::EDefaultFilterBehaviour::Fail)
 	{}
 
 	/** 
@@ -19,8 +20,9 @@ public:
 	 *	@param	InItem	The Item to check 
 	 *	@return			Whether the specified Item passed the filter
 	 */
-	virtual bool PassesFilter( const AActor* InActor ) const override
+	virtual bool PassesFilterImpl(const SceneOutliner::FActorTreeItem& InItem) const
 	{
+		const AActor* InActor = InItem.Actor.Get();
 		if( LayerNames.Num() == 0 )
 		{
 			return false;
diff --git a/Engine/Source/Editor/Layers/Private/SLayerBrowser.cpp b/Engine/Source/Editor/Layers/Private/SLayerBrowser.cpp
--- a/Engine/Source/Editor/Layers/Private/SLayerBrowser.cpp
+++ b/Engine/Source/Editor/Layers/Private/SLayerBrowser.cpp
@@ -126,13 +126,26 @@ void SLayerBrowser::Construct(const FArguments& InArgs)
 	using namespace SceneOutliner;
 	FInitializationOptions InitOptions;
 	{
-		InitOptions.Mode = ESceneOutlinerMode::ActorBrowsing;
-
 		// We hide the header row to keep the UI compact.
 		InitOptions.bShowHeaderRow = false;
 		InitOptions.bShowParentTree = false;
 		InitOptions.bShowCreateNewFolder = false;
-		InitOptions.CustomDelete = FCustomSceneOutlinerDeleteDelegate::CreateSP(this, &SLayerBrowser::RemoveActorsFromSelectedLayer);
+
+		auto CustomDelete = [this](TArray<TWeakPtr<ITreeItem>> Items)
+		{
+			TArray<TWeakObjectPtr<AActor>> Actors;
+			Actors.Reserve(Items.Num());
+			for (const TWeakPtr<ITreeItem>& Item : Items)
+			{
+				if (FActorTreeItem* ActorItem = Item.Pin()->CastTo<FActorTreeItem>())
+				{
+					Actors.Add(ActorItem->Actor);
+				}
+			}
+			RemoveActorsFromSelectedLayer(Actors);
+		};
+
+		InitOptions.CustomDelete = FCustomSceneOutlinerDeleteDelegate::CreateLambda(CustomDelete);
 
 		// Outliner Gutter
 		InitOptions.ColumnMap.Add(FBuiltInColumnTypes::Gutter(), FColumnInfo(EColumnVisibility::Visible, 0) );
@@ -150,7 +163,7 @@ void SLayerBrowser::Construct(const FArguments& InArgs)
 		.BorderImage(FEditorStyle::GetBrush("NoBrush"))
 		.Content()
 		[
-			SceneOutlinerModule.CreateSceneOutliner(InitOptions, FOnActorPicked())
+			SceneOutlinerModule.CreateActorBrowser(InitOptions)
 		];
 
 
diff --git a/Engine/Source/Editor/Layers/Private/SceneOutlinerLayerContentsColumn.cpp b/Engine/Source/Editor/Layers/Private/SceneOutlinerLayerContentsColumn.cpp
--- a/Engine/Source/Editor/Layers/Private/SceneOutlinerLayerContentsColumn.cpp
+++ b/Engine/Source/Editor/Layers/Private/SceneOutlinerLayerContentsColumn.cpp
@@ -7,6 +7,8 @@
 #include "Widgets/Input/SButton.h"
 #include "EditorStyleSet.h"
 
+#include "ActorTreeItem.h"
+
 #define LOCTEXT_NAMESPACE "SceneOutlinerLayerContentsColumn"
 
 
@@ -36,19 +38,23 @@ SHeaderRow::FColumn::FArguments FSceneOutlinerLayerContentsColumn::ConstructHead
 		];
 }
 
-TSharedRef<SWidget> FSceneOutlinerLayerContentsColumn::ConstructRowWidget(const TWeakObjectPtr< AActor >& Actor )
+const TSharedRef<SWidget> FSceneOutlinerLayerContentsColumn::ConstructRowWidget(SceneOutliner::FTreeItemRef TreeItem, const STableRow<SceneOutliner::FTreeItemPtr>& Row)
 {
-	return SNew( SButton )
-		.HAlign( HAlign_Center )
-		.VAlign( VAlign_Center )
-		.ButtonStyle( FEditorStyle::Get(), "LayerBrowserButton" )
-		.ContentPadding( 0 )
-		.OnClicked( this, &FSceneOutlinerLayerContentsColumn::OnRemoveFromLayerClicked, Actor )
-		.ToolTipText( LOCTEXT("RemoveFromLayerButtonText", "Remove from Layer") )
-		[
-			SNew( SImage )
-			.Image( FEditorStyle::GetBrush( TEXT( "LayerBrowser.Actor.RemoveFromLayer" ) ) )
-		];
+	if (SceneOutliner::FActorTreeItem* ActorItem = TreeItem->CastTo<SceneOutliner::FActorTreeItem>())
+	{
+		return SNew(SButton)
+			.HAlign(HAlign_Center)
+			.VAlign(VAlign_Center)
+			.ButtonStyle(FEditorStyle::Get(), "LayerBrowserButton")
+			.ContentPadding(0)
+			.OnClicked(this, &FSceneOutlinerLayerContentsColumn::OnRemoveFromLayerClicked, ActorItem->Actor)
+			.ToolTipText(LOCTEXT("RemoveFromLayerButtonText", "Remove from Layer"))
+			[
+				SNew(SImage)
+				.Image(FEditorStyle::GetBrush(TEXT("LayerBrowser.Actor.RemoveFromLayer")))
+			];
+	}
+	return SNullWidget::NullWidget;
 }
 
 FReply FSceneOutlinerLayerContentsColumn::OnRemoveFromLayerClicked( const TWeakObjectPtr< AActor > Actor )
diff --git a/Engine/Source/Editor/Layers/Private/SceneOutlinerLayerContentsColumn.h b/Engine/Source/Editor/Layers/Private/SceneOutlinerLayerContentsColumn.h
--- a/Engine/Source/Editor/Layers/Private/SceneOutlinerLayerContentsColumn.h
+++ b/Engine/Source/Editor/Layers/Private/SceneOutlinerLayerContentsColumn.h
@@ -37,33 +37,14 @@ public:
 
 	virtual SHeaderRow::FColumn::FArguments ConstructHeaderRowColumn() override;
 
-	virtual const TSharedRef< SWidget > ConstructRowWidget( SceneOutliner::FTreeItemRef TreeItem, const STableRow<SceneOutliner::FTreeItemPtr>& Row ) override
-	{
-		FColumnGenerator Generator(*this);
-		TreeItem->Visit(Generator);
-		return Generator.Widget.ToSharedRef();
-	}
+	virtual const TSharedRef< SWidget > ConstructRowWidget(SceneOutliner::FTreeItemRef TreeItem, const STableRow<SceneOutliner::FTreeItemPtr>& Row) override;
 
 	// End ISceneOutlinerColumn Implementation
 	//////////////////////////////////////////////////////////////////////////
 
 private:
-
 	FReply OnRemoveFromLayerClicked( const TWeakObjectPtr< AActor > Actor );
-
-	TSharedRef<SWidget> ConstructRowWidget(const TWeakObjectPtr< AActor >& Actor );
 	
-	struct FColumnGenerator : SceneOutliner::FColumnGenerator
-	{
-		FSceneOutlinerLayerContentsColumn& Column;
-		FColumnGenerator(FSceneOutlinerLayerContentsColumn& InColumn) : Column(InColumn) {}
-
-		virtual TSharedRef<SWidget> GenerateWidget(SceneOutliner::FActorTreeItem& ActorItem) const override
-		{
-			return Column.ConstructRowWidget(ActorItem.Actor);
-		}
-	};
-	friend FColumnGenerator;
 private:
 
 	/**	 */
diff --git a/Engine/Source/Editor/LevelEditor/Private/LevelEditorContextMenu.cpp b/Engine/Source/Editor/LevelEditor/Private/LevelEditorContextMenu.cpp
--- a/Engine/Source/Editor/LevelEditor/Private/LevelEditorContextMenu.cpp
+++ b/Engine/Source/Editor/LevelEditor/Private/LevelEditorContextMenu.cpp
@@ -38,6 +38,7 @@
 #include "LevelEditorActions.h"
 #include "SceneOutlinerPublicTypes.h"
 #include "SceneOutlinerModule.h"
+#include "ActorTreeItem.h"
 #include "Kismet2/DebuggerCommands.h"
 #include "Styling/SlateIconFinder.h"
 #include "EditorViewportCommands.h"
@@ -1055,13 +1056,12 @@ void FLevelEditorContextMenuImpl::FillActorMenu(UToolMenu* Menu)
 	};
 
 	SceneOutliner::FInitializationOptions InitOptions;
-	{
-		InitOptions.Mode = ESceneOutlinerMode::ActorPicker;			
+	{	
 		InitOptions.bShowHeaderRow = false;
 		InitOptions.bFocusSearchBoxWhenOpened = true;
 
 		// Only display Actors that we can attach too
-		InitOptions.Filters->AddFilterPredicate( SceneOutliner::FActorFilterPredicate::CreateStatic( &FLevelEditorActionCallbacks::IsAttachableActor) );
+		InitOptions.Filters->AddFilterPredicate<SceneOutliner::FActorTreeItem>(SceneOutliner::FActorTreeItem::FFilterPredicate::CreateStatic(&FLevelEditorActionCallbacks::IsAttachableActor));
 	}		
 
 	FToolMenuSection& Section = Menu->AddSection("Actor");
@@ -1085,7 +1085,7 @@ void FLevelEditorContextMenuImpl::FillActorMenu(UToolMenu* Menu)
 			[
 				SNew(SOnlyExpandsBox)
 				[
-					SceneOutlinerModule.CreateSceneOutliner(
+					SceneOutlinerModule.CreateActorPicker(
 						InitOptions,
 						FOnActorPicked::CreateStatic( &FLevelEditorActionCallbacks::AttachToActor )
 						)
diff --git a/Engine/Source/Editor/LevelEditor/Private/LevelEditorToolBar.cpp b/Engine/Source/Editor/LevelEditor/Private/LevelEditorToolBar.cpp
--- a/Engine/Source/Editor/LevelEditor/Private/LevelEditorToolBar.cpp
+++ b/Engine/Source/Editor/LevelEditor/Private/LevelEditorToolBar.cpp
@@ -31,6 +31,7 @@
 #include "Kismet2/DebuggerCommands.h"
 #include "SceneOutlinerPublicTypes.h"
 #include "SceneOutlinerModule.h"
+#include "ActorTreeItem.h"
 #include "SScalabilitySettings.h"
 #include "IContentBrowserSingleton.h"
 #include "ContentBrowserModule.h"
@@ -1999,8 +2000,6 @@ void FLevelEditorToolBar::RegisterCinematicsMenu()
 		// We can't build a list of Matinees and LevelSequenceActors while the current World is a PIE world.
 		FInitializationOptions InitOptions;
 		{
-			InitOptions.Mode = ESceneOutlinerMode::ActorPicker;
-
 			// We hide the header row to keep the UI compact.
 			// @todo: Might be useful to have this sometimes, actually.  Ideally the user could summon it.
 			InitOptions.bShowHeaderRow = false;
@@ -2014,7 +2013,7 @@ void FLevelEditorToolBar::RegisterCinematicsMenu()
 			auto ActorFilter = [&](const AActor* Actor) {
 				return (bAllowMatineeActors && Actor->IsA(AMatineeActor::StaticClass())) || Actor->IsA(ALevelSequenceActor::StaticClass());
 			};
-			InitOptions.Filters->AddFilterPredicate(FActorFilterPredicate::CreateLambda(ActorFilter));
+			InitOptions.Filters->AddFilterPredicate<FActorTreeItem>(FActorTreeItem::FFilterPredicate::CreateLambda(ActorFilter));
 		}
 
 		// actor selector to allow the user to choose an actor
@@ -2025,7 +2024,7 @@ void FLevelEditorToolBar::RegisterCinematicsMenu()
 			.AutoHeight()
 			.MaxHeight(400.0f)
 			[
-				SceneOutlinerModule.CreateSceneOutliner(
+				SceneOutlinerModule.CreateActorPicker(
 					InitOptions,
 					FOnActorPicked::CreateStatic(&FLevelEditorToolBar::OnCinematicsActorPicked))
 			];
diff --git a/Engine/Source/Editor/LevelEditor/Private/SLevelEditor.cpp b/Engine/Source/Editor/LevelEditor/Private/SLevelEditor.cpp
--- a/Engine/Source/Editor/LevelEditor/Private/SLevelEditor.cpp
+++ b/Engine/Source/Editor/LevelEditor/Private/SLevelEditor.cpp
@@ -694,7 +694,6 @@ TSharedRef<SDockTab> SLevelEditor::SpawnLevelEditorTab( const FSpawnTabArgs& Arg
 	{
 		SceneOutliner::FInitializationOptions InitOptions;
 		InitOptions.bShowTransient = true;
-		InitOptions.Mode = ESceneOutlinerMode::ActorBrowsing;
 		{
 			UToolMenus* ToolMenus = UToolMenus::Get();
 			static const FName MenuName = "LevelEditor.LevelEditorSceneOutliner.ContextMenu";
@@ -730,9 +729,8 @@ TSharedRef<SDockTab> SLevelEditor::SpawnLevelEditorTab( const FSpawnTabArgs& Arg
 		FText Label = NSLOCTEXT( "LevelEditor", "SceneOutlinerTabTitle", "World Outliner" );
 
 		FSceneOutlinerModule& SceneOutlinerModule = FModuleManager::Get().LoadModuleChecked<FSceneOutlinerModule>("SceneOutliner");
-		TSharedRef<ISceneOutliner> SceneOutlinerRef = SceneOutlinerModule.CreateSceneOutliner(
-			InitOptions,
-			FOnActorPicked() /* Not used for outliner when in browsing mode */);
+		TSharedRef<ISceneOutliner> SceneOutlinerRef = SceneOutlinerModule.CreateActorBrowser(
+			InitOptions);
 		SceneOutlinerPtr = SceneOutlinerRef;
 
 		return SNew( SDockTab )
diff --git a/Engine/Source/Editor/LevelEditor/Private/SLevelViewport.cpp b/Engine/Source/Editor/LevelEditor/Private/SLevelViewport.cpp
--- a/Engine/Source/Editor/LevelEditor/Private/SLevelViewport.cpp
+++ b/Engine/Source/Editor/LevelEditor/Private/SLevelViewport.cpp
@@ -2544,38 +2544,24 @@ TSharedRef< ISceneOutlinerColumn > SLevelViewport::CreateActorLockSceneOutlinerC
 
 		virtual const TSharedRef< SWidget > ConstructRowWidget( SceneOutliner::FTreeItemRef TreeItem, const STableRow<SceneOutliner::FTreeItemPtr>& InRow ) override
 		{
-			struct FConstructWidget : SceneOutliner::FColumnGenerator
+			if (SceneOutliner::FActorTreeItem* ActorItem = TreeItem->CastTo<SceneOutliner::FActorTreeItem>())
 			{
-				const SLevelViewport* Viewport;
-				FConstructWidget(const SLevelViewport* InViewport) : Viewport(InViewport) {}
-
-				virtual TSharedRef<SWidget> GenerateWidget(SceneOutliner::FActorTreeItem& ActorItem) const override
+				AActor* Actor = ActorItem->Actor.Get();
+				if (!Actor)
 				{
-					AActor* Actor = ActorItem.Actor.Get();
-					if (!Actor)
-					{
-						return SNullWidget::NullWidget;
-					}
-
-					const bool bLocked = Viewport->IsActorLocked(Actor);
-
-					return SNew(SBox)
-						.WidthOverride(SLevelViewport::GetActorLockSceneOutlinerColumnWidth())
-						.Padding(FMargin(2.0f, 0.0f, 0.0f, 0.0f))
-						[
-							SNew(SImage)
-							.Image(FEditorStyle::GetBrush(bLocked ? "PropertyWindow.Locked" : "PropertyWindow.Unlocked"))
-							.ColorAndOpacity(bLocked ? FLinearColor::White : FLinearColor(1.0f, 1.0f, 1.0f, 0.5f))
-						];	
+					return SNullWidget::NullWidget;
 				}
-			};
 
-			FConstructWidget Visitor(Viewport);
-			TreeItem->Visit(Visitor);
+				const bool bLocked = Viewport->IsActorLocked(Actor);
 
-			if (Visitor.Widget.IsValid())
-			{
-				return Visitor.Widget.ToSharedRef();	
+				return SNew(SBox)
+					.WidthOverride(SLevelViewport::GetActorLockSceneOutlinerColumnWidth())
+					.Padding(FMargin(2.0f, 0.0f, 0.0f, 0.0f))
+					[
+						SNew(SImage)
+						.Image(FEditorStyle::GetBrush(bLocked ? "PropertyWindow.Locked" : "PropertyWindow.Unlocked"))
+					.ColorAndOpacity(bLocked ? FLinearColor::White : FLinearColor(1.0f, 1.0f, 1.0f, 0.5f))
+					];
 			}
 			else
 			{
diff --git a/Engine/Source/Editor/MovieSceneTools/Private/TrackEditors/ActorPickerTrackEditor.cpp b/Engine/Source/Editor/MovieSceneTools/Private/TrackEditors/ActorPickerTrackEditor.cpp
--- a/Engine/Source/Editor/MovieSceneTools/Private/TrackEditors/ActorPickerTrackEditor.cpp
+++ b/Engine/Source/Editor/MovieSceneTools/Private/TrackEditors/ActorPickerTrackEditor.cpp
@@ -24,6 +24,7 @@
 #include "MovieSceneObjectBindingIDPicker.h"
 #include "MovieSceneToolHelpers.h"
 #include "SComponentChooser.h"
+#include "ActorTreeItem.h"
 
 #define LOCTEXT_NAMESPACE "FActorPickerTrackEditor"
 
@@ -65,8 +66,7 @@ void FActorPickerTrackEditor::ShowActorSubMenu(FMenuBuilder& MenuBuilder, TArray
 		using namespace SceneOutliner;
 
 		SceneOutliner::FInitializationOptions InitOptions;
-		{
-			InitOptions.Mode = ESceneOutlinerMode::ActorPicker;			
+		{	
 			InitOptions.bShowHeaderRow = false;
 			InitOptions.bFocusSearchBoxWhenOpened = true;
 			InitOptions.bShowTransient = true;
@@ -75,7 +75,7 @@ void FActorPickerTrackEditor::ShowActorSubMenu(FMenuBuilder& MenuBuilder, TArray
 			InitOptions.ColumnMap.Add(FBuiltInColumnTypes::Label(), FColumnInfo(EColumnVisibility::Visible, 0));
 
 			// Only display Actors that we can attach too
-			InitOptions.Filters->AddFilterPredicate( SceneOutliner::FActorFilterPredicate::CreateSP(this, &FActorPickerTrackEditor::IsActorPickable, ObjectBindings[0], Section) );
+			InitOptions.Filters->AddFilterPredicate<SceneOutliner::FActorTreeItem>(SceneOutliner::FActorTreeItem::FFilterPredicate::CreateSP(this, &FActorPickerTrackEditor::IsActorPickable, ObjectBindings[0], Section));
 		}		
 
 		// Actor selector to allow the user to choose a parent actor
@@ -91,7 +91,7 @@ void FActorPickerTrackEditor::ShowActorSubMenu(FMenuBuilder& MenuBuilder, TArray
 				.MaxDesiredHeight(400.0f)
 				.WidthOverride(300.0f)
 				[
-					SceneOutlinerModule.CreateSceneOutliner(
+					SceneOutlinerModule.CreateActorPicker(
 						InitOptions,
 						FOnActorPicked::CreateSP(this, &FActorPickerTrackEditor::ActorPicked, ObjectBindings, Section )
 						)
diff --git a/Engine/Source/Editor/MovieSceneTools/Private/TrackEditors/CameraCutTrackEditor.cpp b/Engine/Source/Editor/MovieSceneTools/Private/TrackEditors/CameraCutTrackEditor.cpp
--- a/Engine/Source/Editor/MovieSceneTools/Private/TrackEditors/CameraCutTrackEditor.cpp
+++ b/Engine/Source/Editor/MovieSceneTools/Private/TrackEditors/CameraCutTrackEditor.cpp
@@ -19,6 +19,7 @@
 #include "ActorEditorUtils.h"
 #include "SceneOutlinerPublicTypes.h"
 #include "SceneOutlinerModule.h"
+#include "ActorTreeItem.h"
 #include "TrackEditorThumbnail/TrackEditorThumbnailPool.h"
 #include "MovieSceneObjectBindingIDPicker.h"
 #include "MovieSceneToolHelpers.h"
@@ -435,7 +436,6 @@ TSharedRef<SWidget> FCameraCutTrackEditor::HandleAddCameraCutComboButtonGetMenuC
 
 			SceneOutliner::FInitializationOptions InitOptions;
 			{
-				InitOptions.Mode = ESceneOutlinerMode::ActorPicker;
 				InitOptions.bShowHeaderRow = false;
 				InitOptions.bFocusSearchBoxWhenOpened = true;
 				InitOptions.bShowTransient = true;
@@ -444,7 +444,7 @@ TSharedRef<SWidget> FCameraCutTrackEditor::HandleAddCameraCutComboButtonGetMenuC
 				InitOptions.ColumnMap.Add(FBuiltInColumnTypes::Label(), FColumnInfo(EColumnVisibility::Visible, 0));
 
 				// Only display Actors that we can attach too
-				InitOptions.Filters->AddFilterPredicate( SceneOutliner::FActorFilterPredicate::CreateRaw(this, &FCameraCutTrackEditor::IsCameraPickable) );
+				InitOptions.Filters->AddFilterPredicate<SceneOutliner::FActorTreeItem>(SceneOutliner::FActorTreeItem::FFilterPredicate::CreateRaw(this, &FCameraCutTrackEditor::IsCameraPickable));
 			}		
 
 			// Actor selector to allow the user to choose a parent actor
@@ -460,7 +460,7 @@ TSharedRef<SWidget> FCameraCutTrackEditor::HandleAddCameraCutComboButtonGetMenuC
 					.MaxDesiredHeight(400.0f)
 					.WidthOverride(300.0f)
 					[
-						SceneOutlinerModule.CreateSceneOutliner(
+						SceneOutlinerModule.CreateActorPicker(
 							InitOptions,
 							FOnActorPicked::CreateSP(this, &FCameraCutTrackEditor::HandleAddCameraCutComboButtonMenuEntryExecute )
 							)
diff --git a/Engine/Source/Editor/PropertyEditor/Private/Presentation/PropertyEditor/PropertyEditor.cpp b/Engine/Source/Editor/PropertyEditor/Private/Presentation/PropertyEditor/PropertyEditor.cpp
--- a/Engine/Source/Editor/PropertyEditor/Private/Presentation/PropertyEditor/PropertyEditor.cpp
+++ b/Engine/Source/Editor/PropertyEditor/Private/Presentation/PropertyEditor/PropertyEditor.cpp
@@ -20,6 +20,7 @@
 #include "EditConditionParser.h"
 #include "EditConditionContext.h"
 #include "Subsystems/AssetEditorSubsystem.h"
+#include "ActorTreeItem.h"
 
 #define LOCTEXT_NAMESPACE "PropertyEditor"
 
@@ -556,7 +557,7 @@ void FPropertyEditor::OnGetActorFiltersForSceneOutliner( TSharedPtr<SceneOutline
 		}
 	};
 
-	OutFilters->AddFilterPredicate( SceneOutliner::FActorFilterPredicate::CreateStatic( &Local::IsFilteredActor, AsShared() ) );
+	OutFilters->AddFilterPredicate<SceneOutliner::FActorTreeItem>(SceneOutliner::FActorTreeItem::FFilterPredicate::CreateStatic( &Local::IsFilteredActor, AsShared() ) );
 }
 
 bool FPropertyEditor::IsPropertyEditingEnabled() const
diff --git a/Engine/Source/Editor/PropertyEditor/Private/UserInterface/PropertyEditor/SPropertyMenuActorPicker.cpp b/Engine/Source/Editor/PropertyEditor/Private/UserInterface/PropertyEditor/SPropertyMenuActorPicker.cpp
--- a/Engine/Source/Editor/PropertyEditor/Private/UserInterface/PropertyEditor/SPropertyMenuActorPicker.cpp
+++ b/Engine/Source/Editor/PropertyEditor/Private/UserInterface/PropertyEditor/SPropertyMenuActorPicker.cpp
@@ -17,6 +17,7 @@
 #include "UserInterface/PropertyEditor/PropertyEditorAssetConstants.h"
 #include "HAL/PlatformApplicationMisc.h"
 #include "SceneOutlinerPublicTypes.h"
+#include "ActorTreeItem.h"
 
 #define LOCTEXT_NAMESPACE "PropertyEditor"
 
@@ -83,19 +84,18 @@ void SPropertyMenuActorPicker::Construct( const FArguments& InArgs )
 		FSceneOutlinerModule& SceneOutlinerModule = FModuleManager::Get().LoadModuleChecked<FSceneOutlinerModule>(TEXT("SceneOutliner"));
 
 		SceneOutliner::FInitializationOptions InitOptions;
-		InitOptions.Mode = ESceneOutlinerMode::ActorPicker;
-		InitOptions.Filters->AddFilterPredicate(ActorFilter);
+		InitOptions.Filters->AddFilterPredicate<SceneOutliner::FActorTreeItem>(ActorFilter);
 		InitOptions.bFocusSearchBoxWhenOpened = true;
 
-		InitOptions.ColumnMap.Add(SceneOutliner::FBuiltInColumnTypes::Label(), SceneOutliner::FColumnInfo(SceneOutliner::EColumnVisibility::Visible, 0) );
-		InitOptions.ColumnMap.Add(SceneOutliner::FBuiltInColumnTypes::ActorInfo(), SceneOutliner::FColumnInfo(SceneOutliner::EColumnVisibility::Visible, 10) );
+		InitOptions.ColumnMap.Add(SceneOutliner::FBuiltInColumnTypes::Label(), SceneOutliner::FColumnInfo(SceneOutliner::EColumnVisibility::Visible, 0));
+		InitOptions.ColumnMap.Add(SceneOutliner::FBuiltInColumnTypes::ActorInfo(), SceneOutliner::FColumnInfo(SceneOutliner::EColumnVisibility::Visible, 10));
 		
 		MenuContent =
 			SNew(SBox)
 			.WidthOverride(PropertyEditorAssetConstants::SceneOutlinerWindowSize.X)
 			.HeightOverride(PropertyEditorAssetConstants::SceneOutlinerWindowSize.Y)
 			[
-				SceneOutlinerModule.CreateSceneOutliner(InitOptions, FOnActorPicked::CreateSP(this, &SPropertyMenuActorPicker::OnActorSelected))
+				SceneOutlinerModule.CreateActorPicker(InitOptions, FOnActorPicked::CreateSP(this, &SPropertyMenuActorPicker::OnActorSelected))
 			];
 
 		MenuBuilder.AddWidget(MenuContent.ToSharedRef(), FText::GetEmpty(), true);
diff --git a/Engine/Source/Editor/PropertyEditor/Private/UserInterface/PropertyEditor/SPropertyMenuComponentPicker.cpp b/Engine/Source/Editor/PropertyEditor/Private/UserInterface/PropertyEditor/SPropertyMenuComponentPicker.cpp
--- a/Engine/Source/Editor/PropertyEditor/Private/UserInterface/PropertyEditor/SPropertyMenuComponentPicker.cpp
+++ b/Engine/Source/Editor/PropertyEditor/Private/UserInterface/PropertyEditor/SPropertyMenuComponentPicker.cpp
@@ -19,7 +19,8 @@
 #include "UserInterface/PropertyEditor/PropertyEditorAssetConstants.h"
 #include "HAL/PlatformApplicationMisc.h"
 #include "SceneOutlinerPublicTypes.h"
-#include "ITreeItem.h"
+#include "ActorTreeItem.h"
+#include "ComponentTreeItem.h"
 
 #define LOCTEXT_NAMESPACE "PropertyEditor"
 
@@ -80,11 +81,45 @@ void SPropertyMenuComponentPicker::Construct(const FArguments& InArgs)
 		FSceneOutlinerModule& SceneOutlinerModule = FModuleManager::Get().LoadModuleChecked<FSceneOutlinerModule>(TEXT("SceneOutliner"));
 
 		SceneOutliner::FInitializationOptions InitOptions;
-		InitOptions.Mode = ESceneOutlinerMode::ComponentPicker;
 		InitOptions.bFocusSearchBoxWhenOpened = true;
 
-		TSharedRef<SceneOutliner::FOutlinerPredicateFilter> Filter = MakeShareable(new SceneOutliner::FOutlinerPredicateFilter(ActorFilter, SceneOutliner::EDefaultFilterBehaviour::Fail));
-		Filter->ComponentPred = ComponentFilter;
+		struct FPickerFilter : public SceneOutliner::FOutlinerFilter
+		{
+			FPickerFilter(const FOnShouldFilterActor& InActorFilter, const FOnShouldFilterComponent& InComponentFilter)
+				: FOutlinerFilter(SceneOutliner::EDefaultFilterBehaviour::Fail)
+				, ActorFilter(InActorFilter)
+				, ComponentFilter(InComponentFilter)
+			{}
+
+			virtual bool PassesFilter(const SceneOutliner::ITreeItem& InItem) const override
+			{
+				if (const SceneOutliner::FActorTreeItem* ActorItem = InItem.CastTo<SceneOutliner::FActorTreeItem>())
+				{
+					return ActorItem->IsValid() && ActorFilter.Execute(ActorItem->Actor.Get());
+				}
+				else if (const SceneOutliner::FComponentTreeItem* ComponentItem = InItem.CastTo<SceneOutliner::FComponentTreeItem>())
+				{
+					return ComponentItem->IsValid() && ComponentFilter.Execute(ComponentItem->Component.Get());
+				}
+
+				return DefaultBehaviour == SceneOutliner::EDefaultFilterBehaviour::Pass;
+			}
+
+			virtual bool GetInteractiveState(const SceneOutliner::ITreeItem& InItem) const override
+			{
+				// All components which pass the filter are interactive
+				if (const SceneOutliner::FComponentTreeItem* ComponentItem = InItem.CastTo<SceneOutliner::FComponentTreeItem>())
+				{
+					return true;
+				}
+				return DefaultBehaviour == SceneOutliner::EDefaultFilterBehaviour::Pass;
+			}
+
+			FOnShouldFilterActor ActorFilter;
+			FOnShouldFilterComponent ComponentFilter;
+		};
+
+		TSharedRef<SceneOutliner::FOutlinerFilter> Filter = MakeShared<FPickerFilter>(ActorFilter, ComponentFilter);
 		InitOptions.Filters->Add(Filter);
 
 		InitOptions.ColumnMap.Add(SceneOutliner::FBuiltInColumnTypes::Label(), SceneOutliner::FColumnInfo(SceneOutliner::EColumnVisibility::Visible, 0));
@@ -97,7 +132,7 @@ void SPropertyMenuComponentPicker::Construct(const FArguments& InArgs)
 				SNew(SBorder)
 				.BorderImage(FEditorStyle::GetBrush("Menu.Background"))
 				[
-					SceneOutlinerModule.CreateSceneOutliner(InitOptions, FOnSceneOutlinerItemPicked::CreateSP(this, &SPropertyMenuComponentPicker::OnItemSelected))
+					SceneOutlinerModule.CreateComponentPicker(InitOptions, FOnComponentPicked::CreateSP(this, &SPropertyMenuComponentPicker::OnItemSelected))
 				]
 			];
 
@@ -198,17 +233,9 @@ void SPropertyMenuComponentPicker::OnClear()
 	OnClose.ExecuteIfBound();
 }
 
-void SPropertyMenuComponentPicker::OnItemSelected(TSharedRef<SceneOutliner::ITreeItem> InItem)
+void SPropertyMenuComponentPicker::OnItemSelected(UActorComponent* Component)
 {
-	InItem->Visit(
-		SceneOutliner::FFunctionalVisitor()
-		.Component([&](const SceneOutliner::FComponentTreeItem& ComponentItem)
-		{
-			if (UActorComponent* Component = ComponentItem.Component.Get())
-			{
-				SetValue(Component);
-			}
-		}));
+	SetValue(Component);
 	OnClose.ExecuteIfBound();
 }
 
diff --git a/Engine/Source/Editor/PropertyEditor/Private/UserInterface/PropertyEditor/SPropertyMenuComponentPicker.h b/Engine/Source/Editor/PropertyEditor/Private/UserInterface/PropertyEditor/SPropertyMenuComponentPicker.h
--- a/Engine/Source/Editor/PropertyEditor/Private/UserInterface/PropertyEditor/SPropertyMenuComponentPicker.h
+++ b/Engine/Source/Editor/PropertyEditor/Private/UserInterface/PropertyEditor/SPropertyMenuComponentPicker.h
@@ -65,9 +65,9 @@ private:
 
 	/**
 	 * Delegate for handling selection by the actor picker.
-	 * @param	InActor	The chosen actor
+	 * @param	Component The chosen component
 	 */
-	void OnItemSelected(TSharedRef<SceneOutliner::ITreeItem> InItem);
+	void OnItemSelected(UActorComponent* Component);
 
 	/**
 	 * Set the value of the asset referenced by this property editor.
diff --git a/Engine/Source/Editor/PropertyEditor/Private/UserInterface/PropertyEditor/SPropertySceneOutliner.cpp b/Engine/Source/Editor/PropertyEditor/Private/UserInterface/PropertyEditor/SPropertySceneOutliner.cpp
--- a/Engine/Source/Editor/PropertyEditor/Private/UserInterface/PropertyEditor/SPropertySceneOutliner.cpp
+++ b/Engine/Source/Editor/PropertyEditor/Private/UserInterface/PropertyEditor/SPropertySceneOutliner.cpp
@@ -56,7 +56,6 @@ TSharedRef<SWidget> SPropertySceneOutliner::OnGenerateSceneOutliner()
 	FSceneOutlinerModule& SceneOutlinerModule = FModuleManager::Get().LoadModuleChecked<FSceneOutlinerModule>(TEXT("SceneOutliner"));
 
 	SceneOutliner::FInitializationOptions InitOptions;
-	InitOptions.Mode = ESceneOutlinerMode::ActorPicker;
 	OnGetActorFilters.ExecuteIfBound( InitOptions.Filters );
 
 	TSharedRef<SWidget> MenuContent = 
@@ -67,7 +66,7 @@ TSharedRef<SWidget> SPropertySceneOutliner::OnGenerateSceneOutliner()
 			SNew( SBorder )
 			.BorderImage( FEditorStyle::GetBrush("Menu.Background") )
 			[
-				SceneOutlinerModule.CreateSceneOutliner(InitOptions, FOnActorPicked::CreateSP(this, &SPropertySceneOutliner::OnActorSelectedFromOutliner))
+				SceneOutlinerModule.CreateActorPicker(InitOptions, FOnActorPicked::CreateSP(this, &SPropertySceneOutliner::OnActorSelectedFromOutliner))
 			]
 		];
 
