Removed SceneOutliner namespace

diff --git a/Engine/Plugins/Editor/MeshEditor/Source/MeshEditor/FractureToolComponent.cpp b/Engine/Plugins/Editor/MeshEditor/Source/MeshEditor/FractureToolComponent.cpp
--- a/Engine/Plugins/Editor/MeshEditor/Source/MeshEditor/FractureToolComponent.cpp
+++ b/Engine/Plugins/Editor/MeshEditor/Source/MeshEditor/FractureToolComponent.cpp
@@ -163,7 +163,7 @@ void UFractureToolComponent::SetSelectedBones(UEditableMesh* EditableMesh, int32
 					EditBoneColor.SetSelectedBones(RevisedSelected);
 					EditBoneColor.SetHighlightedBones(Highlighted);
 
-					SceneOutliner::FSceneOutlinerDelegates::Get().OnComponentSelectionChanged.Broadcast(GeometryCollectionComponent);
+					FSceneOutlinerDelegates::Get().OnComponentSelectionChanged.Broadcast(GeometryCollectionComponent);
 				}
 			}
 
diff --git a/Engine/Plugins/Editor/MeshEditor/Source/PolygonModeling/GeometryCollectionCommandCommon.cpp b/Engine/Plugins/Editor/MeshEditor/Source/PolygonModeling/GeometryCollectionCommandCommon.cpp
--- a/Engine/Plugins/Editor/MeshEditor/Source/PolygonModeling/GeometryCollectionCommandCommon.cpp
+++ b/Engine/Plugins/Editor/MeshEditor/Source/PolygonModeling/GeometryCollectionCommandCommon.cpp
@@ -217,7 +217,7 @@ void FGeometryCollectionCommandCommon::UpdateExplodedView(class IMeshEditorModeE
 	// Update the exploded view in the UI based on the current exploded view slider position
 	FFractureToolDelegates::Get().OnUpdateExplodedView.Broadcast(static_cast<uint8>(ResetType), static_cast<uint8>(MeshEditorMode.GetFractureSettings()->CommonSettings->ViewMode));
 
-	SceneOutliner::FSceneOutlinerDelegates::Get().OnComponentsUpdated.Broadcast();
+	FSceneOutlinerDelegates::Get().OnComponentsUpdated.Broadcast();
 }
 
 UGeometryCollectionComponent* FGeometryCollectionCommandCommon::GetGeometryCollectionComponent(UEditableMesh* SourceMesh)
diff --git a/Engine/Plugins/Enterprise/DataprepEditor/Source/DataprepEditor/Private/DataprepEditor.cpp b/Engine/Plugins/Enterprise/DataprepEditor/Source/DataprepEditor/Private/DataprepEditor.cpp
--- a/Engine/Plugins/Enterprise/DataprepEditor/Source/DataprepEditor/Private/DataprepEditor.cpp
+++ b/Engine/Plugins/Enterprise/DataprepEditor/Source/DataprepEditor/Private/DataprepEditor.cpp
@@ -1147,8 +1147,8 @@ void FDataprepEditor::RefreshColumnsForPreviewSystem()
 {
 	if ( PreviewSystem->HasObservedObjects() )
 	{
-		SceneOutliner->RemoveColumn( SceneOutliner::FBuiltInColumnTypes::ActorInfo() );
-		SceneOutliner::FColumnInfo ColumnInfo( SceneOutliner::EColumnVisibility::Visible, 100, FCreateSceneOutlinerColumn::CreateLambda( [Preview = PreviewSystem](ISceneOutliner& InSceneOutliner) -> TSharedRef< ISceneOutlinerColumn >
+		SceneOutliner->RemoveColumn( FSceneOutlinerBuiltInColumnTypes::ActorInfo() );
+		FSceneOutlinerColumnInfo ColumnInfo(ESceneOutlinerColumnVisibility::Visible, 100, FCreateSceneOutlinerColumn::CreateLambda([Preview = PreviewSystem](ISceneOutliner& InSceneOutliner) -> TSharedRef< ISceneOutlinerColumn >
 			{
 				return MakeShared<FDataprepPreviewOutlinerColumn>( InSceneOutliner, Preview );
 			} ) );
@@ -1160,9 +1160,9 @@ void FDataprepEditor::RefreshColumnsForPreviewSystem()
 	{
 		SceneOutliner->RemoveColumn( FDataprepPreviewOutlinerColumn::ColumnID );
 		FSceneOutlinerModule& SceneOutlinerModule = FModuleManager::Get().LoadModuleChecked<FSceneOutlinerModule>("SceneOutliner");
-		SceneOutliner::FColumnInfo* ActorInfoColumPtr = SceneOutlinerModule.DefaultColumnMap.Find( SceneOutliner::FBuiltInColumnTypes::ActorInfo() );
+		FSceneOutlinerColumnInfo* ActorInfoColumPtr = SceneOutlinerModule.DefaultColumnMap.Find( FSceneOutlinerBuiltInColumnTypes::ActorInfo() );
 		check( ActorInfoColumPtr );
-		SceneOutliner->AddColumn( SceneOutliner::FBuiltInColumnTypes::ActorInfo(), *ActorInfoColumPtr );
+		SceneOutliner->AddColumn( FSceneOutlinerBuiltInColumnTypes::ActorInfo(), *ActorInfoColumPtr );
 
 		AssetPreviewView->RemoveColumn( FDataprepPreviewAssetColumn::ColumnID );
 	}
diff --git a/Engine/Plugins/Enterprise/DataprepEditor/Source/DataprepEditor/Private/DataprepEditor.h b/Engine/Plugins/Enterprise/DataprepEditor/Source/DataprepEditor/Private/DataprepEditor.h
--- a/Engine/Plugins/Enterprise/DataprepEditor/Source/DataprepEditor/Private/DataprepEditor.h
+++ b/Engine/Plugins/Enterprise/DataprepEditor/Source/DataprepEditor/Private/DataprepEditor.h
@@ -124,7 +124,7 @@ public:
 	int32 GetCountOfPreviewedSteps() const;
 
 	/** Handles change to selection in SceneOutliner */
-	void OnSceneOutlinerSelectionChanged(SceneOutliner::FTreeItemPtr ItemPtr, ESelectInfo::Type SelectionMode);
+	void OnSceneOutlinerSelectionChanged(FSceneOutlinerTreeItemPtr ItemPtr, ESelectInfo::Type SelectionMode);
 
 private:
 	void BindCommands();
diff --git a/Engine/Plugins/Enterprise/DataprepEditor/Source/DataprepEditor/Private/DataprepEditorOutliner.cpp b/Engine/Plugins/Enterprise/DataprepEditor/Source/DataprepEditor/Private/DataprepEditorOutliner.cpp
--- a/Engine/Plugins/Enterprise/DataprepEditor/Source/DataprepEditor/Private/DataprepEditorOutliner.cpp
+++ b/Engine/Plugins/Enterprise/DataprepEditor/Source/DataprepEditor/Private/DataprepEditorOutliner.cpp
@@ -24,7 +24,6 @@
 namespace DataprepEditorSceneOutlinerUtils
 {
 	/** Specifics for the scene outliner */
-	using namespace SceneOutliner;
 
 	/**
 	 * Use this struct to get the selection from the scene outliner
@@ -33,7 +32,7 @@ namespace DataprepEditorSceneOutlinerUtils
 	{
 		mutable TSet<TWeakObjectPtr<UObject>> Selection;
 
-		void TryAdd(const ITreeItem& Item) const
+		void TryAdd(const ISceneOutlinerTreeItem& Item) const
 		{
 			if (const FActorTreeItem* ActorItem = Item.CastTo<FActorTreeItem>())
 			{
@@ -46,7 +45,7 @@ namespace DataprepEditorSceneOutlinerUtils
 		}
 	};
 
-	void SetVisibility(ITreeItem& Item, bool bIsVisible, TWeakPtr<SDataprepEditorViewport> Viewport)
+	void SetVisibility(ISceneOutlinerTreeItem& Item, bool bIsVisible, TWeakPtr<SDataprepEditorViewport> Viewport)
 	{
 		if (FActorTreeItem* ActorItem = Item.CastTo<FActorTreeItem>())
 		{
@@ -145,12 +144,12 @@ namespace DataprepEditorSceneOutlinerUtils
 				];
 		}
 
-		virtual const TSharedRef<SWidget> ConstructRowWidget(FTreeItemRef TreeItem, const STableRow<FTreeItemPtr>& Row) override;
+		virtual const TSharedRef<SWidget> ConstructRowWidget(FSceneOutlinerTreeItemRef TreeItem, const STableRow<FSceneOutlinerTreeItemPtr>& Row) override;
 
-		virtual void SortItems(TArray<FTreeItemPtr>& RootItems, const EColumnSortMode::Type SortMode) const override {}
+		virtual void SortItems(TArray<FSceneOutlinerTreeItemPtr>& RootItems, const EColumnSortMode::Type SortMode) const override {}
 
 		/** Check whether the specified item is visible */
-		bool IsItemVisible(const ITreeItem& Item) { return VisibilityCache.GetVisibility(Item); }
+		bool IsItemVisible(const ISceneOutlinerTreeItem& Item) { return VisibilityCache.GetVisibility(Item); }
 
 		TWeakPtr<SDataprepEditorViewport> GetViewport() const
 		{
@@ -164,7 +163,7 @@ namespace DataprepEditorSceneOutlinerUtils
 		TWeakPtr<SDataprepEditorViewport> WeakSceneViewport;
 
 		/** Get and cache visibility for items. Cached per-frame to avoid expensive recursion. */
-		FGetVisibilityCache VisibilityCache;
+		FSceneOutlinerVisibilityCache VisibilityCache;
 	};
 
 	/** Widget responsible for managing the visibility for a single actor */
@@ -175,7 +174,7 @@ namespace DataprepEditorSceneOutlinerUtils
 		SLATE_END_ARGS()
 
 		/** Construct this widget */
-		void Construct(const FArguments& InArgs, TWeakPtr<FPreviewSceneOutlinerGutter> InWeakColumn, TWeakPtr<ISceneOutliner> InWeakOutliner, TWeakPtr<ITreeItem> InWeakTreeItem)
+		void Construct(const FArguments& InArgs, TWeakPtr<FPreviewSceneOutlinerGutter> InWeakColumn, TWeakPtr<ISceneOutliner> InWeakOutliner, TWeakPtr<ISceneOutlinerTreeItem> InWeakTreeItem)
 		{
 			WeakTreeItem = InWeakTreeItem;
 			WeakOutliner = InWeakOutliner;
@@ -307,7 +306,7 @@ namespace DataprepEditorSceneOutlinerUtils
 		}
 
 		/** Check if the specified item is visible */
-		static bool IsVisible(const FTreeItemPtr& Item, const TSharedPtr<FPreviewSceneOutlinerGutter>& Column)
+		static bool IsVisible(const FSceneOutlinerTreeItemPtr& Item, const TSharedPtr<FPreviewSceneOutlinerGutter>& Column)
 		{
 			return Column.IsValid() && Item.IsValid() ? Column->IsItemVisible(*Item) : false;
 		}
@@ -321,7 +320,7 @@ namespace DataprepEditorSceneOutlinerUtils
 		/** Set the actor this widget is responsible for to be hidden or shown */
 		void SetIsVisible(const bool bVisible)
 		{
-			TSharedPtr<ITreeItem> TreeItem = WeakTreeItem.Pin();
+			TSharedPtr<ISceneOutlinerTreeItem> TreeItem = WeakTreeItem.Pin();
 			TSharedPtr<ISceneOutliner> Outliner = WeakOutliner.Pin();
 
 			if (TreeItem.IsValid() && Outliner.IsValid() && IsVisible() != bVisible)
@@ -335,7 +334,7 @@ namespace DataprepEditorSceneOutlinerUtils
 		}
 
 		/** The tree item we relate to */
-		TWeakPtr<ITreeItem> WeakTreeItem;
+		TWeakPtr<ISceneOutlinerTreeItem> WeakTreeItem;
 
 		/** Reference back to the outliner so we can set visibility of a whole selection */
 		TWeakPtr<ISceneOutliner> WeakOutliner;
@@ -347,7 +346,7 @@ namespace DataprepEditorSceneOutlinerUtils
 		TUniquePtr<FScopedTransaction> UndoTransaction;
 	};
 
-	const TSharedRef<SWidget> FPreviewSceneOutlinerGutter::ConstructRowWidget(FTreeItemRef TreeItem, const STableRow<FTreeItemPtr>& Row)
+	const TSharedRef<SWidget> FPreviewSceneOutlinerGutter::ConstructRowWidget(FSceneOutlinerTreeItemRef TreeItem, const STableRow<FSceneOutlinerTreeItemPtr>& Row)
 	{
 		return SNew(SHorizontalBox)
 			+ SHorizontalBox::Slot()
@@ -365,12 +364,12 @@ void FDataprepEditor::CreateScenePreviewTab()
 {
 	FSceneOutlinerModule& SceneOutlinerModule = FModuleManager::Get().LoadModuleChecked<FSceneOutlinerModule>("SceneOutliner");
 	
-	SceneOutliner::FCreateOutlinerMode ModeFactory = SceneOutliner::FCreateOutlinerMode::CreateLambda([SpecifiedWorld = PreviewWorld.Get(), DataprepEditor = StaticCastSharedRef<FDataprepEditor>(AsShared())](SceneOutliner::SSceneOutliner* Outliner)
+	FCreateSceneOutlinerMode ModeFactory = FCreateSceneOutlinerMode::CreateLambda([SpecifiedWorld = PreviewWorld.Get(), DataprepEditor = StaticCastSharedRef<FDataprepEditor>(AsShared())](SSceneOutliner* Outliner)
 		{
 			return new FDataprepEditorOutlinerMode(Outliner, DataprepEditor, SpecifiedWorld);
 		});
 
-	SceneOutliner::FInitializationOptions SceneOutlinerOptions;
+	FSceneOutlinerInitializationOptions SceneOutlinerOptions;
 	SceneOutlinerOptions.ModeFactory = ModeFactory;
 
 	SceneOutliner = SceneOutlinerModule.CreateSceneOutliner(SceneOutlinerOptions);
@@ -382,8 +381,8 @@ void FDataprepEditor::CreateScenePreviewTab()
 		return TSharedRef< ISceneOutlinerColumn >(MakeShareable(new DataprepEditorSceneOutlinerUtils::FPreviewSceneOutlinerGutter(Outliner, SceneViewportView.ToSharedRef())));
 	};
 
-	SceneOutliner::FColumnInfo ColumnInfo;
-	ColumnInfo.Visibility = SceneOutliner::EColumnVisibility::Visible;
+	FSceneOutlinerColumnInfo ColumnInfo;
+	ColumnInfo.Visibility = ESceneOutlinerColumnVisibility::Visible;
 	ColumnInfo.PriorityIndex = 0;
 	ColumnInfo.Factory.BindLambda([&](ISceneOutliner& Outliner)
 	{
@@ -393,7 +392,7 @@ void FDataprepEditor::CreateScenePreviewTab()
 	SceneOutliner->AddColumn(DataprepEditorSceneOutlinerUtils::FPreviewSceneOutlinerGutter::GetID(), ColumnInfo);
 
 	// Add the default columns
-	const SceneOutliner::FSharedOutlinerData& SharedData = SceneOutliner->GetSharedData();
+	const FSharedSceneOutlinerData& SharedData = SceneOutliner->GetSharedData();
 	for (auto& DefaultColumn : SceneOutlinerModule.DefaultColumnMap)
 	{
 		SceneOutliner->AddColumn(DefaultColumn.Key, DefaultColumn.Value);
@@ -413,13 +412,11 @@ void FDataprepEditor::CreateScenePreviewTab()
 		];
 }
 
-void FDataprepEditor::OnSceneOutlinerSelectionChanged(SceneOutliner::FTreeItemPtr ItemPtr, ESelectInfo::Type SelectionMode)
+void FDataprepEditor::OnSceneOutlinerSelectionChanged(FSceneOutlinerTreeItemPtr ItemPtr, ESelectInfo::Type SelectionMode)
 {
-	using namespace SceneOutliner;
-
 	DataprepEditorSceneOutlinerUtils::FGetSelectionFromSceneOutliner Selector;
 
-	for (FTreeItemPtr Item : SceneOutliner->GetTree().GetSelectedItems())
+	for (FSceneOutlinerTreeItemPtr Item : SceneOutliner->GetTree().GetSelectedItems())
 	{
 		Selector.TryAdd(*Item);
 	}
diff --git a/Engine/Plugins/Enterprise/DataprepEditor/Source/DataprepEditor/Private/DataprepEditorOutlinerMode.h b/Engine/Plugins/Enterprise/DataprepEditor/Source/DataprepEditor/Private/DataprepEditorOutlinerMode.h
--- a/Engine/Plugins/Enterprise/DataprepEditor/Source/DataprepEditor/Private/DataprepEditorOutlinerMode.h
+++ b/Engine/Plugins/Enterprise/DataprepEditor/Source/DataprepEditor/Private/DataprepEditorOutlinerMode.h
@@ -8,8 +8,6 @@
 
 namespace DataprepEditorSceneOutlinerUtils
 {
-	using namespace SceneOutliner;
-
 	/**
 	* Use this struct to match the scene outliers selection to a dataprep editor selection
 	*/
@@ -20,7 +18,7 @@ namespace DataprepEditorSceneOutlinerUtils
 		{
 		};
 
-		bool operator()(const ITreeItem& Item) const
+		bool operator()(const ISceneOutlinerTreeItem& Item) const
 		{
 			if (const FActorTreeItem* ActorItem = Item.CastTo<FActorTreeItem>())
 			{
@@ -44,18 +42,18 @@ namespace DataprepEditorSceneOutlinerUtils
 	};
 }
 
-class FDataprepEditorOutlinerMode : public SceneOutliner::FActorMode
+class FDataprepEditorOutlinerMode : public FActorMode
 {
 public:
-	FDataprepEditorOutlinerMode(SceneOutliner::SSceneOutliner* InSceneOutliner, TWeakPtr<FDataprepEditor> InDataprepEditor, TWeakObjectPtr<UWorld> InSpecifiedWorldToDisplay = nullptr)
-		: SceneOutliner::FActorMode(InSceneOutliner, true, InSpecifiedWorldToDisplay)
+	FDataprepEditorOutlinerMode(SSceneOutliner* InSceneOutliner, TWeakPtr<FDataprepEditor> InDataprepEditor, TWeakObjectPtr<UWorld> InSpecifiedWorldToDisplay = nullptr)
+		: FActorMode(InSceneOutliner, true, InSpecifiedWorldToDisplay)
 		, DataprepEditorPtr(InDataprepEditor)
 	{}
 
-	virtual bool CanRenameItem(const SceneOutliner::ITreeItem& Item) const override { return false; }
+	virtual bool CanRenameItem(const ISceneOutlinerTreeItem& Item) const override { return false; }
 	virtual ESelectionMode::Type GetSelectionMode() const override { return ESelectionMode::Multi; }
 
-	virtual void OnItemSelectionChanged(SceneOutliner::FTreeItemPtr Item, ESelectInfo::Type SelectionType, const SceneOutliner::FItemSelection& Selection) override
+	virtual void OnItemSelectionChanged(FSceneOutlinerTreeItemPtr Item, ESelectInfo::Type SelectionType, const FSceneOutlinerItemSelection& Selection) override
 	{
 		auto DataprepEditor = DataprepEditorPtr.Pin();
 		if (DataprepEditor)
diff --git a/Engine/Plugins/Enterprise/DataprepEditor/Source/DataprepEditor/Private/PreviewSystem/DataprepPreviewSceneOutlinerColumn.cpp b/Engine/Plugins/Enterprise/DataprepEditor/Source/DataprepEditor/Private/PreviewSystem/DataprepPreviewSceneOutlinerColumn.cpp
--- a/Engine/Plugins/Enterprise/DataprepEditor/Source/DataprepEditor/Private/PreviewSystem/DataprepPreviewSceneOutlinerColumn.cpp
+++ b/Engine/Plugins/Enterprise/DataprepEditor/Source/DataprepEditor/Private/PreviewSystem/DataprepPreviewSceneOutlinerColumn.cpp
@@ -20,13 +20,13 @@
 namespace DataprepPreviewOutlinerColumnUtils
 {
 	
-	UObject* GetObjectPtr(const SceneOutliner::ITreeItem& Item)
+	UObject* GetObjectPtr(const ISceneOutlinerTreeItem& Item)
 	{
-		if (const SceneOutliner::FActorTreeItem* ActorItem = Item.CastTo<SceneOutliner::FActorTreeItem>())
+		if (const FActorTreeItem* ActorItem = Item.CastTo<FActorTreeItem>())
 		{
 			return ActorItem->Actor.Get();
 		}
-		else if (const SceneOutliner::FComponentTreeItem* ComponentItem = Item.CastTo<SceneOutliner::FComponentTreeItem>())
+		else if (const FComponentTreeItem* ComponentItem = Item.CastTo<FComponentTreeItem>())
 		{
 			return ComponentItem->Component.Get();
 		}
@@ -58,7 +58,7 @@ SHeaderRow::FColumn::FArguments FDataprepPreviewOutlinerColumn::ConstructHeaderR
 		.FillWidth( 5.0f );
 }
 
-const TSharedRef<SWidget> FDataprepPreviewOutlinerColumn::ConstructRowWidget(SceneOutliner::FTreeItemRef TreeItem, const STableRow<SceneOutliner::FTreeItemPtr>& Row)
+const TSharedRef<SWidget> FDataprepPreviewOutlinerColumn::ConstructRowWidget(FSceneOutlinerTreeItemRef TreeItem, const STableRow<FSceneOutlinerTreeItemPtr>& Row)
 {
 	if ( UObject* Object = DataprepPreviewOutlinerColumnUtils::GetObjectPtr(*TreeItem))
 	{
@@ -72,7 +72,7 @@ const TSharedRef<SWidget> FDataprepPreviewOutlinerColumn::ConstructRowWidget(Sce
 	return SNullWidget::NullWidget;
 }
 
-void FDataprepPreviewOutlinerColumn::PopulateSearchStrings(const SceneOutliner::ITreeItem& Item, TArray< FString >& OutSearchStrings) const
+void FDataprepPreviewOutlinerColumn::PopulateSearchStrings(const ISceneOutlinerTreeItem& Item, TArray< FString >& OutSearchStrings) const
 {
 	if (UObject* Object = DataprepPreviewOutlinerColumnUtils::GetObjectPtr(Item))
 	{
@@ -83,9 +83,9 @@ void FDataprepPreviewOutlinerColumn::PopulateSearchStrings(const SceneOutliner::
 	}
 }
 
-void FDataprepPreviewOutlinerColumn::SortItems(TArray<SceneOutliner::FTreeItemPtr>& OutItems, const EColumnSortMode::Type SortMode) const
+void FDataprepPreviewOutlinerColumn::SortItems(TArray<FSceneOutlinerTreeItemPtr>& OutItems, const EColumnSortMode::Type SortMode) const
 {
-	OutItems.Sort([this, SortMode](const SceneOutliner::FTreeItemPtr& First, const SceneOutliner::FTreeItemPtr& Second)
+	OutItems.Sort([this, SortMode](const FSceneOutlinerTreeItemPtr& First, const FSceneOutlinerTreeItemPtr& Second)
 	{
 			UObject* FirstObject = DataprepPreviewOutlinerColumnUtils::GetObjectPtr(*First);
 			UObject* SecondObject = DataprepPreviewOutlinerColumnUtils::GetObjectPtr(*Second);
diff --git a/Engine/Plugins/Enterprise/DataprepEditor/Source/DataprepEditor/Private/PreviewSystem/DataprepPreviewSceneOutlinerColumn.h b/Engine/Plugins/Enterprise/DataprepEditor/Source/DataprepEditor/Private/PreviewSystem/DataprepPreviewSceneOutlinerColumn.h
--- a/Engine/Plugins/Enterprise/DataprepEditor/Source/DataprepEditor/Private/PreviewSystem/DataprepPreviewSceneOutlinerColumn.h
+++ b/Engine/Plugins/Enterprise/DataprepEditor/Source/DataprepEditor/Private/PreviewSystem/DataprepPreviewSceneOutlinerColumn.h
@@ -20,13 +20,13 @@ public:
 
 	virtual SHeaderRow::FColumn::FArguments ConstructHeaderRowColumn() override;
 
-	virtual const TSharedRef<SWidget> ConstructRowWidget(SceneOutliner::FTreeItemRef TreeItem, const STableRow<SceneOutliner::FTreeItemPtr>& Row) override;
+	virtual const TSharedRef<SWidget> ConstructRowWidget(FSceneOutlinerTreeItemRef TreeItem, const STableRow<FSceneOutlinerTreeItemPtr>& Row) override;
 
-	virtual void PopulateSearchStrings(const SceneOutliner::ITreeItem& Item, TArray< FString >& OutSearchStrings) const override;
+	virtual void PopulateSearchStrings(const ISceneOutlinerTreeItem& Item, TArray< FString >& OutSearchStrings) const override;
 
 	virtual bool SupportsSorting() const override { return true; }
 
-	virtual void SortItems(TArray<SceneOutliner::FTreeItemPtr>& OutItems, const EColumnSortMode::Type SortMode) const override;
+	virtual void SortItems(TArray<FSceneOutlinerTreeItemPtr>& OutItems, const EColumnSortMode::Type SortMode) const override;
 
 	const static FName ColumnID;
 
diff --git a/Engine/Plugins/Enterprise/VariantManager/Source/VariantManager/Private/DisplayNodes/VariantManagerActorNode.cpp b/Engine/Plugins/Enterprise/VariantManager/Source/VariantManager/Private/DisplayNodes/VariantManagerActorNode.cpp
--- a/Engine/Plugins/Enterprise/VariantManager/Source/VariantManager/Private/DisplayNodes/VariantManagerActorNode.cpp
+++ b/Engine/Plugins/Enterprise/VariantManager/Source/VariantManager/Private/DisplayNodes/VariantManagerActorNode.cpp
@@ -476,7 +476,6 @@ const UClass* FVariantManagerActorNode::GetClassForObjectBinding() const
 void FVariantManagerActorNode::AddAssignActorSubMenu(FMenuBuilder& MenuBuilder)
 {
 	// Copied from FSequencer::AssignActor
-	using namespace SceneOutliner;
 
 	// If we're showing this menu, we know for a fact only our actor node is selected,
 	// so we only have to check our variant
@@ -506,7 +505,7 @@ void FVariantManagerActorNode::AddAssignActorSubMenu(FMenuBuilder& MenuBuilder)
 							 TAttribute<FText>(this, &FVariantManagerActorNode::GetRebindToSelectedTooltip));
 
 	// Set up a menu entry to assign an actor to the object binding node
-	FInitializationOptions InitOptions;
+	FSceneOutlinerInitializationOptions InitOptions;
 	{
 		// We hide the header row to keep the UI compact.
 		InitOptions.bShowHeaderRow = false;
@@ -514,7 +513,7 @@ void FVariantManagerActorNode::AddAssignActorSubMenu(FMenuBuilder& MenuBuilder)
 		InitOptions.bShowCreateNewFolder = false;
 		InitOptions.bFocusSearchBoxWhenOpened = true;
 		// Only want the actor label column
-		InitOptions.ColumnMap.Add(FBuiltInColumnTypes::Label(), FColumnInfo(EColumnVisibility::Visible, 0));
+		InitOptions.ColumnMap.Add(FSceneOutlinerBuiltInColumnTypes::Label(), FSceneOutlinerColumnInfo(ESceneOutlinerColumnVisibility::Visible, 0));
 
 		// Only display actors that are not possessed already
 		InitOptions.Filters->AddFilterPredicate<FActorTreeItem>(FActorTreeItem::FFilterPredicate::CreateLambda( IsActorValidForAssignment ));
diff --git a/Engine/Plugins/Enterprise/VariantManager/Source/VariantManager/Private/SVariantManager.cpp b/Engine/Plugins/Enterprise/VariantManager/Source/VariantManager/Private/SVariantManager.cpp
--- a/Engine/Plugins/Enterprise/VariantManager/Source/VariantManager/Private/SVariantManager.cpp
+++ b/Engine/Plugins/Enterprise/VariantManager/Source/VariantManager/Private/SVariantManager.cpp
@@ -2581,8 +2581,6 @@ FReply SVariantManager::OnAddVariantSetClicked()
 
 FReply SVariantManager::OnSummonAddActorMenu()
 {
-	using namespace SceneOutliner;
-
 	TSharedPtr<FVariantManager> VarMan = VariantManagerPtr.Pin();
 	if (!VarMan.IsValid())
 	{
@@ -2650,12 +2648,12 @@ FReply SVariantManager::OnSummonAddActorMenu()
 		return !CommonActorSet.Contains(InActor);
 	};
 
-	FInitializationOptions InitOptions;
+	FSceneOutlinerInitializationOptions InitOptions;
 	InitOptions.bShowHeaderRow = true;
 	InitOptions.bShowSearchBox = true;
 	InitOptions.bShowCreateNewFolder = false;
 	InitOptions.bFocusSearchBoxWhenOpened = true;
-	InitOptions.ColumnMap.Add(FBuiltInColumnTypes::Label(), FColumnInfo(EColumnVisibility::Visible, 0));
+	InitOptions.ColumnMap.Add(FSceneOutlinerBuiltInColumnTypes::Label(), FSceneOutlinerColumnInfo(ESceneOutlinerColumnVisibility::Visible, 0));
 	InitOptions.Filters->AddFilterPredicate<FActorTreeItem>(FActorTreeItem::FFilterPredicate::CreateLambda( IsActorValidForAssignment ));
 
 	// Create mini scene outliner menu
diff --git a/Engine/Plugins/Experimental/GeometryCollectionPlugin/Source/GeometryCollectionEditor/Private/GeometryCollection/GeometryCollectionCommands.cpp b/Engine/Plugins/Experimental/GeometryCollectionPlugin/Source/GeometryCollectionEditor/Private/GeometryCollection/GeometryCollectionCommands.cpp
--- a/Engine/Plugins/Experimental/GeometryCollectionPlugin/Source/GeometryCollectionEditor/Private/GeometryCollection/GeometryCollectionCommands.cpp
+++ b/Engine/Plugins/Experimental/GeometryCollectionPlugin/Source/GeometryCollectionEditor/Private/GeometryCollection/GeometryCollectionCommands.cpp
@@ -657,7 +657,7 @@ void FGeometryCollectionCommands::SelectLessThenVolume(const TArray<FString>& Ar
 					EditBoneColor.SetSelectedBones(SelectedBones);
 					EditBoneColor.SetHighlightedBones(SelectedBones);
 
- 					SceneOutliner::FSceneOutlinerDelegates::Get().OnComponentSelectionChanged.Broadcast(Actor->GetGeometryCollectionComponent());
+ 					FSceneOutlinerDelegates::Get().OnComponentSelectionChanged.Broadcast(Actor->GetGeometryCollectionComponent());
 				}
 			}
 		}
diff --git a/Engine/Plugins/MovieScene/LevelSequenceEditor/Source/LevelSequenceEditor/Private/Misc/LevelSequenceEditorActorBinding.cpp b/Engine/Plugins/MovieScene/LevelSequenceEditor/Source/LevelSequenceEditor/Private/Misc/LevelSequenceEditorActorBinding.cpp
--- a/Engine/Plugins/MovieScene/LevelSequenceEditor/Source/LevelSequenceEditor/Private/Misc/LevelSequenceEditorActorBinding.cpp
+++ b/Engine/Plugins/MovieScene/LevelSequenceEditor/Source/LevelSequenceEditor/Private/Misc/LevelSequenceEditorActorBinding.cpp
@@ -97,10 +97,8 @@ void FLevelSequenceEditorActorBinding::AddPossessActorMenuExtensions(FMenuBuilde
 	
 	MenuBuilder.BeginSection("ChooseActorSection", LOCTEXT("ChooseActor", "Choose Actor:"));
 
-	using namespace SceneOutliner;
-
 	// Set up a menu entry to add any arbitrary actor to the sequencer
-	FInitializationOptions InitOptions;
+	FSceneOutlinerInitializationOptions InitOptions;
 	{
 		// We hide the header row to keep the UI compact.
 		InitOptions.bShowHeaderRow = false;
@@ -109,7 +107,7 @@ void FLevelSequenceEditorActorBinding::AddPossessActorMenuExtensions(FMenuBuilde
 		InitOptions.bFocusSearchBoxWhenOpened = true;
 
 		// Only want the actor label column
-		InitOptions.ColumnMap.Add(FBuiltInColumnTypes::Label(), FColumnInfo(EColumnVisibility::Visible, 0));
+		InitOptions.ColumnMap.Add(FSceneOutlinerBuiltInColumnTypes::Label(), FSceneOutlinerColumnInfo(ESceneOutlinerColumnVisibility::Visible, 0));
 
 		// Only display actors that are not possessed already
 		InitOptions.Filters->AddFilterPredicate<FActorTreeItem>(FActorTreeItem::FFilterPredicate::CreateLambda(IsActorValidForPossession, ExistingPossessedObjects));
diff --git a/Engine/Plugins/VirtualProduction/Takes/Source/TakeRecorderSources/Private/TakeRecorderSourcesModule.cpp b/Engine/Plugins/VirtualProduction/Takes/Source/TakeRecorderSources/Private/TakeRecorderSourcesModule.cpp
--- a/Engine/Plugins/VirtualProduction/Takes/Source/TakeRecorderSources/Private/TakeRecorderSourcesModule.cpp
+++ b/Engine/Plugins/VirtualProduction/Takes/Source/TakeRecorderSources/Private/TakeRecorderSourcesModule.cpp
@@ -175,8 +175,6 @@ struct FActorTakeRecorderDropHandler : ITakeRecorderDropHandler
 
 	virtual bool CanHandleOperation(TSharedPtr<FDragDropOperation> InOperation, UTakeRecorderSources* Sources) override
 	{
-		using namespace SceneOutliner;
-
 		if (InOperation)
 		{
 			if (InOperation->IsOfType<FActorDragDropOp>())
@@ -194,8 +192,6 @@ struct FActorTakeRecorderDropHandler : ITakeRecorderDropHandler
 
 	TArray<AActor*> GetValidDropActors(TSharedPtr<FDragDropOperation> InOperation, UTakeRecorderSources* Sources)
 	{
-		using namespace SceneOutliner;
-
 		FActorDragDropOp*  ActorDrag = nullptr;
 		FFolderDragDropOp* FolderDrag = nullptr;
 
@@ -394,10 +390,8 @@ public:
 
 		MenuBuilder.BeginSection("ChooseActorSection", LOCTEXT("ChooseActor", "Choose Actor:"));
 		{
-			using namespace SceneOutliner;
-
 			// Set up a menu entry to add any arbitrary actor to the sequencer
-			FInitializationOptions InitOptions;
+			FSceneOutlinerInitializationOptions InitOptions;
 			{
 				// We hide the header row to keep the UI compact.
 				InitOptions.bShowHeaderRow = false;
@@ -406,7 +400,7 @@ public:
 				InitOptions.bFocusSearchBoxWhenOpened = true;
 
 				// Only want the actor label column
-				InitOptions.ColumnMap.Add(FBuiltInColumnTypes::Label(), FColumnInfo(EColumnVisibility::Visible, 0));
+				InitOptions.ColumnMap.Add(FSceneOutlinerBuiltInColumnTypes::Label(), FSceneOutlinerColumnInfo(ESceneOutlinerColumnVisibility::Visible, 0));
 
 				// Only display actors that are not possessed already
 				InitOptions.Filters->AddFilterPredicate<FActorTreeItem>(FActorTreeItem::FFilterPredicate::CreateLambda(OutlinerFilterPredicate));
diff --git a/Engine/Source/Editor/Layers/Private/ActorsAssignedToSpecificLayersFilter.h b/Engine/Source/Editor/Layers/Private/ActorsAssignedToSpecificLayersFilter.h
--- a/Engine/Source/Editor/Layers/Private/ActorsAssignedToSpecificLayersFilter.h
+++ b/Engine/Source/Editor/Layers/Private/ActorsAssignedToSpecificLayersFilter.h
@@ -6,13 +6,13 @@
 #include "ActorTreeItem.h"
 #include "GameFramework/Actor.h"
 
-class FActorsAssignedToSpecificLayersFilter : public SceneOutliner::TOutlinerFilter<SceneOutliner::FActorTreeItem>, public TSharedFromThis< FActorsAssignedToSpecificLayersFilter >
+class FActorsAssignedToSpecificLayersFilter : public TSceneOutlinerFilter<FActorTreeItem>, public TSharedFromThis< FActorsAssignedToSpecificLayersFilter >
 {
 public:
 
 	/** By default, any types not handled by this filter will fail */
 	FActorsAssignedToSpecificLayersFilter()
-		: TOutlinerFilter<SceneOutliner::FActorTreeItem>(SceneOutliner::EDefaultFilterBehaviour::Fail)
+		: TSceneOutlinerFilter<FActorTreeItem>(FSceneOutlinerFilter::EDefaultBehaviour::Fail)
 	{}
 
 	/** 
@@ -21,7 +21,7 @@ public:
 	 *	@param	InItem	The Item to check 
 	 *	@return			Whether the specified Item passed the filter
 	 */
-	virtual bool PassesFilterImpl(const SceneOutliner::FActorTreeItem& InItem) const
+	virtual bool PassesFilterImpl(const FActorTreeItem& InItem) const
 	{
 		const AActor* InActor = InItem.Actor.Get();
 		if( LayerNames.Num() == 0 )
diff --git a/Engine/Source/Editor/Layers/Private/SLayerBrowser.cpp b/Engine/Source/Editor/Layers/Private/SLayerBrowser.cpp
--- a/Engine/Source/Editor/Layers/Private/SLayerBrowser.cpp
+++ b/Engine/Source/Editor/Layers/Private/SLayerBrowser.cpp
@@ -123,19 +123,18 @@ void SLayerBrowser::Construct(const FArguments& InArgs)
 	//	Layer Contents Section
 	FSceneOutlinerModule& SceneOutlinerModule = FModuleManager::LoadModuleChecked< FSceneOutlinerModule >("SceneOutliner");
 	
-	using namespace SceneOutliner;
-	FInitializationOptions InitOptions;
+	FSceneOutlinerInitializationOptions InitOptions;
 	{
 		// We hide the header row to keep the UI compact.
 		InitOptions.bShowHeaderRow = false;
 		InitOptions.bShowParentTree = false;
 		InitOptions.bShowCreateNewFolder = false;
 
-		auto CustomDelete = [this](TArray<TWeakPtr<ITreeItem>> Items)
+		auto CustomDelete = [this](TArray<TWeakPtr<ISceneOutlinerTreeItem>> Items)
 		{
 			TArray<TWeakObjectPtr<AActor>> Actors;
 			Actors.Reserve(Items.Num());
-			for (const TWeakPtr<ITreeItem>& Item : Items)
+			for (const TWeakPtr<ISceneOutlinerTreeItem>& Item : Items)
 			{
 				if (FActorTreeItem* ActorItem = Item.Pin()->CastTo<FActorTreeItem>())
 				{
@@ -148,11 +147,11 @@ void SLayerBrowser::Construct(const FArguments& InArgs)
 		InitOptions.CustomDelete = FCustomSceneOutlinerDeleteDelegate::CreateLambda(CustomDelete);
 
 		// Outliner Gutter
-		InitOptions.ColumnMap.Add(FBuiltInColumnTypes::Gutter(), FColumnInfo(EColumnVisibility::Visible, 0) );
+		InitOptions.ColumnMap.Add(FSceneOutlinerBuiltInColumnTypes::Gutter(), FSceneOutlinerColumnInfo(ESceneOutlinerColumnVisibility::Visible, 0) );
 		// Actor Label
-		InitOptions.ColumnMap.Add(FBuiltInColumnTypes::Label(), FColumnInfo(EColumnVisibility::Visible, 10) );
+		InitOptions.ColumnMap.Add(FSceneOutlinerBuiltInColumnTypes::Label(), FSceneOutlinerColumnInfo(ESceneOutlinerColumnVisibility::Visible, 10) );
 		// Layer Contents
-		InitOptions.ColumnMap.Add(FSceneOutlinerLayerContentsColumn::GetID(), FColumnInfo(SceneOutliner::EColumnVisibility::Visible, 20,
+		InitOptions.ColumnMap.Add(FSceneOutlinerLayerContentsColumn::GetID(), FSceneOutlinerColumnInfo(ESceneOutlinerColumnVisibility::Visible, 20,
 			FCreateSceneOutlinerColumn::CreateSP( this, &SLayerBrowser::CreateCustomLayerColumn )) );
 
 		InitOptions.Filters->Add(SelectedLayersFilter);
diff --git a/Engine/Source/Editor/Layers/Private/SceneOutlinerLayerContentsColumn.cpp b/Engine/Source/Editor/Layers/Private/SceneOutlinerLayerContentsColumn.cpp
--- a/Engine/Source/Editor/Layers/Private/SceneOutlinerLayerContentsColumn.cpp
+++ b/Engine/Source/Editor/Layers/Private/SceneOutlinerLayerContentsColumn.cpp
@@ -38,9 +38,9 @@ SHeaderRow::FColumn::FArguments FSceneOutlinerLayerContentsColumn::ConstructHead
 		];
 }
 
-const TSharedRef<SWidget> FSceneOutlinerLayerContentsColumn::ConstructRowWidget(SceneOutliner::FTreeItemRef TreeItem, const STableRow<SceneOutliner::FTreeItemPtr>& Row)
+const TSharedRef<SWidget> FSceneOutlinerLayerContentsColumn::ConstructRowWidget(FSceneOutlinerTreeItemRef TreeItem, const STableRow<FSceneOutlinerTreeItemPtr>& Row)
 {
-	if (SceneOutliner::FActorTreeItem* ActorItem = TreeItem->CastTo<SceneOutliner::FActorTreeItem>())
+	if (FActorTreeItem* ActorItem = TreeItem->CastTo<FActorTreeItem>())
 	{
 		return SNew(SButton)
 			.HAlign(HAlign_Center)
diff --git a/Engine/Source/Editor/Layers/Private/SceneOutlinerLayerContentsColumn.h b/Engine/Source/Editor/Layers/Private/SceneOutlinerLayerContentsColumn.h
--- a/Engine/Source/Editor/Layers/Private/SceneOutlinerLayerContentsColumn.h
+++ b/Engine/Source/Editor/Layers/Private/SceneOutlinerLayerContentsColumn.h
@@ -37,7 +37,7 @@ public:
 
 	virtual SHeaderRow::FColumn::FArguments ConstructHeaderRowColumn() override;
 
-	virtual const TSharedRef< SWidget > ConstructRowWidget(SceneOutliner::FTreeItemRef TreeItem, const STableRow<SceneOutliner::FTreeItemPtr>& Row) override;
+	virtual const TSharedRef< SWidget > ConstructRowWidget(FSceneOutlinerTreeItemRef TreeItem, const STableRow<FSceneOutlinerTreeItemPtr>& Row) override;
 
 	// End ISceneOutlinerColumn Implementation
 	//////////////////////////////////////////////////////////////////////////
diff --git a/Engine/Source/Editor/LevelEditor/Private/LevelEditorContextMenu.cpp b/Engine/Source/Editor/LevelEditor/Private/LevelEditorContextMenu.cpp
--- a/Engine/Source/Editor/LevelEditor/Private/LevelEditorContextMenu.cpp
+++ b/Engine/Source/Editor/LevelEditor/Private/LevelEditorContextMenu.cpp
@@ -1055,13 +1055,13 @@ void FLevelEditorContextMenuImpl::FillActorMenu(UToolMenu* Menu)
 		}
 	};
 
-	SceneOutliner::FInitializationOptions InitOptions;
+	FSceneOutlinerInitializationOptions InitOptions;
 	{	
 		InitOptions.bShowHeaderRow = false;
 		InitOptions.bFocusSearchBoxWhenOpened = true;
 
 		// Only display Actors that we can attach too
-		InitOptions.Filters->AddFilterPredicate<SceneOutliner::FActorTreeItem>(SceneOutliner::FActorTreeItem::FFilterPredicate::CreateStatic(&FLevelEditorActionCallbacks::IsAttachableActor));
+		InitOptions.Filters->AddFilterPredicate<FActorTreeItem>(FActorTreeItem::FFilterPredicate::CreateStatic(&FLevelEditorActionCallbacks::IsAttachableActor));
 	}		
 
 	FToolMenuSection& Section = Menu->AddSection("Actor");
diff --git a/Engine/Source/Editor/LevelEditor/Private/LevelEditorToolBar.cpp b/Engine/Source/Editor/LevelEditor/Private/LevelEditorToolBar.cpp
--- a/Engine/Source/Editor/LevelEditor/Private/LevelEditorToolBar.cpp
+++ b/Engine/Source/Editor/LevelEditor/Private/LevelEditorToolBar.cpp
@@ -1995,10 +1995,8 @@ void FLevelEditorToolBar::RegisterCinematicsMenu()
 			return;
 		}
 
-		using namespace SceneOutliner;
-
 		// We can't build a list of Matinees and LevelSequenceActors while the current World is a PIE world.
-		FInitializationOptions InitOptions;
+		FSceneOutlinerInitializationOptions InitOptions;
 		{
 			// We hide the header row to keep the UI compact.
 			// @todo: Might be useful to have this sometimes, actually.  Ideally the user could summon it.
@@ -2006,8 +2004,8 @@ void FLevelEditorToolBar::RegisterCinematicsMenu()
 			InitOptions.bShowSearchBox = false;
 			InitOptions.bShowCreateNewFolder = false;
 
-			InitOptions.ColumnMap.Add(FBuiltInColumnTypes::Label(), FColumnInfo(EColumnVisibility::Visible, 0));
-			InitOptions.ColumnMap.Add(FBuiltInColumnTypes::ActorInfo(), FColumnInfo(EColumnVisibility::Visible, 10));
+			InitOptions.ColumnMap.Add(FSceneOutlinerBuiltInColumnTypes::Label(), FSceneOutlinerColumnInfo(ESceneOutlinerColumnVisibility::Visible, 0));
+			InitOptions.ColumnMap.Add(FSceneOutlinerBuiltInColumnTypes::ActorInfo(), FSceneOutlinerColumnInfo(ESceneOutlinerColumnVisibility::Visible, 10));
 
 			// Only display Matinee and MovieScene actors
 			auto ActorFilter = [&](const AActor* Actor) {
diff --git a/Engine/Source/Editor/LevelEditor/Private/SLevelEditor.cpp b/Engine/Source/Editor/LevelEditor/Private/SLevelEditor.cpp
--- a/Engine/Source/Editor/LevelEditor/Private/SLevelEditor.cpp
+++ b/Engine/Source/Editor/LevelEditor/Private/SLevelEditor.cpp
@@ -692,7 +692,7 @@ TSharedRef<SDockTab> SLevelEditor::SpawnLevelEditorTab( const FSpawnTabArgs& Arg
 	}
 	else if( TabIdentifier == LevelEditorTabIds::LevelEditorSceneOutliner)
 	{
-		SceneOutliner::FInitializationOptions InitOptions;
+		FSceneOutlinerInitializationOptions InitOptions;
 		InitOptions.bShowTransient = true;
 		{
 			UToolMenus* ToolMenus = UToolMenus::Get();
diff --git a/Engine/Source/Editor/LevelEditor/Private/SLevelViewport.cpp b/Engine/Source/Editor/LevelEditor/Private/SLevelViewport.cpp
--- a/Engine/Source/Editor/LevelEditor/Private/SLevelViewport.cpp
+++ b/Engine/Source/Editor/LevelEditor/Private/SLevelViewport.cpp
@@ -2542,9 +2542,9 @@ TSharedRef< ISceneOutlinerColumn > SLevelViewport::CreateActorLockSceneOutlinerC
 				];
 		}
 
-		virtual const TSharedRef< SWidget > ConstructRowWidget( SceneOutliner::FTreeItemRef TreeItem, const STableRow<SceneOutliner::FTreeItemPtr>& InRow ) override
+		virtual const TSharedRef< SWidget > ConstructRowWidget( FSceneOutlinerTreeItemRef TreeItem, const STableRow<FSceneOutlinerTreeItemPtr>& InRow ) override
 		{
-			if (SceneOutliner::FActorTreeItem* ActorItem = TreeItem->CastTo<SceneOutliner::FActorTreeItem>())
+			if (FActorTreeItem* ActorItem = TreeItem->CastTo<FActorTreeItem>())
 			{
 				AActor* Actor = ActorItem->Actor.Get();
 				if (!Actor)
diff --git a/Engine/Source/Editor/MovieSceneTools/Private/TrackEditors/ActorPickerTrackEditor.cpp b/Engine/Source/Editor/MovieSceneTools/Private/TrackEditors/ActorPickerTrackEditor.cpp
--- a/Engine/Source/Editor/MovieSceneTools/Private/TrackEditors/ActorPickerTrackEditor.cpp
+++ b/Engine/Source/Editor/MovieSceneTools/Private/TrackEditors/ActorPickerTrackEditor.cpp
@@ -63,19 +63,17 @@ void FActorPickerTrackEditor::ShowActorSubMenu(FMenuBuilder& MenuBuilder, TArray
 	auto CreateNewBinding = 
 		[this, ObjectBindings, Section](FMenuBuilder& SubMenuBuilder)
 	{
-		using namespace SceneOutliner;
-
-		SceneOutliner::FInitializationOptions InitOptions;
+		FSceneOutlinerInitializationOptions InitOptions;
 		{	
 			InitOptions.bShowHeaderRow = false;
 			InitOptions.bFocusSearchBoxWhenOpened = true;
 			InitOptions.bShowTransient = true;
 			InitOptions.bShowCreateNewFolder = false;
 			// Only want the actor label column
-			InitOptions.ColumnMap.Add(FBuiltInColumnTypes::Label(), FColumnInfo(EColumnVisibility::Visible, 0));
+			InitOptions.ColumnMap.Add(FSceneOutlinerBuiltInColumnTypes::Label(), FSceneOutlinerColumnInfo(ESceneOutlinerColumnVisibility::Visible, 0));
 
 			// Only display Actors that we can attach too
-			InitOptions.Filters->AddFilterPredicate<SceneOutliner::FActorTreeItem>(SceneOutliner::FActorTreeItem::FFilterPredicate::CreateSP(this, &FActorPickerTrackEditor::IsActorPickable, ObjectBindings[0], Section));
+			InitOptions.Filters->AddFilterPredicate<FActorTreeItem>(FActorTreeItem::FFilterPredicate::CreateSP(this, &FActorPickerTrackEditor::IsActorPickable, ObjectBindings[0], Section));
 		}		
 
 		// Actor selector to allow the user to choose a parent actor
diff --git a/Engine/Source/Editor/MovieSceneTools/Private/TrackEditors/CameraCutTrackEditor.cpp b/Engine/Source/Editor/MovieSceneTools/Private/TrackEditors/CameraCutTrackEditor.cpp
--- a/Engine/Source/Editor/MovieSceneTools/Private/TrackEditors/CameraCutTrackEditor.cpp
+++ b/Engine/Source/Editor/MovieSceneTools/Private/TrackEditors/CameraCutTrackEditor.cpp
@@ -432,19 +432,17 @@ TSharedRef<SWidget> FCameraCutTrackEditor::HandleAddCameraCutComboButtonGetMenuC
 	auto CreateNewCamera =
 		[this](FMenuBuilder& SubMenuBuilder)
 		{
-			using namespace SceneOutliner;
-
-			SceneOutliner::FInitializationOptions InitOptions;
+			FSceneOutlinerInitializationOptions InitOptions;
 			{
 				InitOptions.bShowHeaderRow = false;
 				InitOptions.bFocusSearchBoxWhenOpened = true;
 				InitOptions.bShowTransient = true;
 				InitOptions.bShowCreateNewFolder = false;
 				// Only want the actor label column
-				InitOptions.ColumnMap.Add(FBuiltInColumnTypes::Label(), FColumnInfo(EColumnVisibility::Visible, 0));
+				InitOptions.ColumnMap.Add(FSceneOutlinerBuiltInColumnTypes::Label(), FSceneOutlinerColumnInfo(ESceneOutlinerColumnVisibility::Visible, 0));
 
 				// Only display Actors that we can attach too
-				InitOptions.Filters->AddFilterPredicate<SceneOutliner::FActorTreeItem>(SceneOutliner::FActorTreeItem::FFilterPredicate::CreateRaw(this, &FCameraCutTrackEditor::IsCameraPickable));
+				InitOptions.Filters->AddFilterPredicate<FActorTreeItem>(FActorTreeItem::FFilterPredicate::CreateRaw(this, &FCameraCutTrackEditor::IsCameraPickable));
 			}		
 
 			// Actor selector to allow the user to choose a parent actor
diff --git a/Engine/Source/Editor/PropertyEditor/Private/Presentation/PropertyEditor/PropertyEditor.cpp b/Engine/Source/Editor/PropertyEditor/Private/Presentation/PropertyEditor/PropertyEditor.cpp
--- a/Engine/Source/Editor/PropertyEditor/Private/Presentation/PropertyEditor/PropertyEditor.cpp
+++ b/Engine/Source/Editor/PropertyEditor/Private/Presentation/PropertyEditor/PropertyEditor.cpp
@@ -539,7 +539,7 @@ void FPropertyEditor::OnActorSelected( AActor* InActor )
 	OnAssetSelected(InActor);
 }
 
-void FPropertyEditor::OnGetActorFiltersForSceneOutliner( TSharedPtr<SceneOutliner::FOutlinerFilters>& OutFilters )
+void FPropertyEditor::OnGetActorFiltersForSceneOutliner( TSharedPtr<FSceneOutlinerFilters>& OutFilters )
 {
 	struct Local
 	{
@@ -557,7 +557,7 @@ void FPropertyEditor::OnGetActorFiltersForSceneOutliner( TSharedPtr<SceneOutline
 		}
 	};
 
-	OutFilters->AddFilterPredicate<SceneOutliner::FActorTreeItem>(SceneOutliner::FActorTreeItem::FFilterPredicate::CreateStatic( &Local::IsFilteredActor, AsShared() ) );
+	OutFilters->AddFilterPredicate<FActorTreeItem>(FActorTreeItem::FFilterPredicate::CreateStatic( &Local::IsFilteredActor, AsShared() ) );
 }
 
 bool FPropertyEditor::IsPropertyEditingEnabled() const
diff --git a/Engine/Source/Editor/PropertyEditor/Private/Presentation/PropertyEditor/PropertyEditor.h b/Engine/Source/Editor/PropertyEditor/Private/Presentation/PropertyEditor/PropertyEditor.h
--- a/Engine/Source/Editor/PropertyEditor/Private/Presentation/PropertyEditor/PropertyEditor.h
+++ b/Engine/Source/Editor/PropertyEditor/Private/Presentation/PropertyEditor/PropertyEditor.h
@@ -87,7 +87,7 @@ public:
 	void OnGetClassesForAssetPicker( TArray<const UClass*>& OutClasses );
 	void OnAssetSelected( const FAssetData& AssetData );
 	void OnActorSelected( AActor* InActor );
-	void OnGetActorFiltersForSceneOutliner( TSharedPtr<SceneOutliner::FOutlinerFilters>& OutFilters );
+	void OnGetActorFiltersForSceneOutliner( TSharedPtr<FSceneOutlinerFilters>& OutFilters );
 	void EditConfigHierarchy();
 
 	/**	In an ideal world we wouldn't expose these */
diff --git a/Engine/Source/Editor/PropertyEditor/Private/UserInterface/PropertyEditor/SPropertyMenuActorPicker.cpp b/Engine/Source/Editor/PropertyEditor/Private/UserInterface/PropertyEditor/SPropertyMenuActorPicker.cpp
--- a/Engine/Source/Editor/PropertyEditor/Private/UserInterface/PropertyEditor/SPropertyMenuActorPicker.cpp
+++ b/Engine/Source/Editor/PropertyEditor/Private/UserInterface/PropertyEditor/SPropertyMenuActorPicker.cpp
@@ -83,12 +83,12 @@ void SPropertyMenuActorPicker::Construct( const FArguments& InArgs )
 
 		FSceneOutlinerModule& SceneOutlinerModule = FModuleManager::Get().LoadModuleChecked<FSceneOutlinerModule>(TEXT("SceneOutliner"));
 
-		SceneOutliner::FInitializationOptions InitOptions;
-		InitOptions.Filters->AddFilterPredicate<SceneOutliner::FActorTreeItem>(ActorFilter);
+		FSceneOutlinerInitializationOptions InitOptions;
+		InitOptions.Filters->AddFilterPredicate<FActorTreeItem>(ActorFilter);
 		InitOptions.bFocusSearchBoxWhenOpened = true;
 
-		InitOptions.ColumnMap.Add(SceneOutliner::FBuiltInColumnTypes::Label(), SceneOutliner::FColumnInfo(SceneOutliner::EColumnVisibility::Visible, 0));
-		InitOptions.ColumnMap.Add(SceneOutliner::FBuiltInColumnTypes::ActorInfo(), SceneOutliner::FColumnInfo(SceneOutliner::EColumnVisibility::Visible, 10));
+		InitOptions.ColumnMap.Add(FSceneOutlinerBuiltInColumnTypes::Label(), FSceneOutlinerColumnInfo(ESceneOutlinerColumnVisibility::Visible, 0));
+		InitOptions.ColumnMap.Add(FSceneOutlinerBuiltInColumnTypes::ActorInfo(), FSceneOutlinerColumnInfo(ESceneOutlinerColumnVisibility::Visible, 10));
 		
 		MenuContent =
 			SNew(SBox)
diff --git a/Engine/Source/Editor/PropertyEditor/Private/UserInterface/PropertyEditor/SPropertyMenuComponentPicker.cpp b/Engine/Source/Editor/PropertyEditor/Private/UserInterface/PropertyEditor/SPropertyMenuComponentPicker.cpp
--- a/Engine/Source/Editor/PropertyEditor/Private/UserInterface/PropertyEditor/SPropertyMenuComponentPicker.cpp
+++ b/Engine/Source/Editor/PropertyEditor/Private/UserInterface/PropertyEditor/SPropertyMenuComponentPicker.cpp
@@ -80,49 +80,49 @@ void SPropertyMenuComponentPicker::Construct(const FArguments& InArgs)
 
 		FSceneOutlinerModule& SceneOutlinerModule = FModuleManager::Get().LoadModuleChecked<FSceneOutlinerModule>(TEXT("SceneOutliner"));
 
-		SceneOutliner::FInitializationOptions InitOptions;
+		FSceneOutlinerInitializationOptions InitOptions;
 		InitOptions.bFocusSearchBoxWhenOpened = true;
 
-		struct FPickerFilter : public SceneOutliner::FOutlinerFilter
+		struct FPickerFilter : public FSceneOutlinerFilter
 		{
 			FPickerFilter(const FOnShouldFilterActor& InActorFilter, const FOnShouldFilterComponent& InComponentFilter)
-				: FOutlinerFilter(SceneOutliner::EDefaultFilterBehaviour::Fail)
+				: FSceneOutlinerFilter(FSceneOutlinerFilter::EDefaultBehaviour::Fail)
 				, ActorFilter(InActorFilter)
 				, ComponentFilter(InComponentFilter)
 			{}
 
-			virtual bool PassesFilter(const SceneOutliner::ITreeItem& InItem) const override
+			virtual bool PassesFilter(const ISceneOutlinerTreeItem& InItem) const override
 			{
-				if (const SceneOutliner::FActorTreeItem* ActorItem = InItem.CastTo<SceneOutliner::FActorTreeItem>())
+				if (const FActorTreeItem* ActorItem = InItem.CastTo<FActorTreeItem>())
 				{
 					return ActorItem->IsValid() && ActorFilter.Execute(ActorItem->Actor.Get());
 				}
-				else if (const SceneOutliner::FComponentTreeItem* ComponentItem = InItem.CastTo<SceneOutliner::FComponentTreeItem>())
+				else if (const FComponentTreeItem* ComponentItem = InItem.CastTo<FComponentTreeItem>())
 				{
 					return ComponentItem->IsValid() && ComponentFilter.Execute(ComponentItem->Component.Get());
 				}
 
-				return DefaultBehaviour == SceneOutliner::EDefaultFilterBehaviour::Pass;
+				return DefaultBehaviour == FSceneOutlinerFilter::EDefaultBehaviour::Pass;
 			}
 
-			virtual bool GetInteractiveState(const SceneOutliner::ITreeItem& InItem) const override
+			virtual bool GetInteractiveState(const ISceneOutlinerTreeItem& InItem) const override
 			{
 				// All components which pass the filter are interactive
-				if (const SceneOutliner::FComponentTreeItem* ComponentItem = InItem.CastTo<SceneOutliner::FComponentTreeItem>())
+				if (const FComponentTreeItem* ComponentItem = InItem.CastTo<FComponentTreeItem>())
 				{
 					return true;
 				}
-				return DefaultBehaviour == SceneOutliner::EDefaultFilterBehaviour::Pass;
+				return DefaultBehaviour == FSceneOutlinerFilter::EDefaultBehaviour::Pass;
 			}
 
 			FOnShouldFilterActor ActorFilter;
 			FOnShouldFilterComponent ComponentFilter;
 		};
 
-		TSharedRef<SceneOutliner::FOutlinerFilter> Filter = MakeShared<FPickerFilter>(ActorFilter, ComponentFilter);
+		TSharedRef<FSceneOutlinerFilter> Filter = MakeShared<FPickerFilter>(ActorFilter, ComponentFilter);
 		InitOptions.Filters->Add(Filter);
 
-		InitOptions.ColumnMap.Add(SceneOutliner::FBuiltInColumnTypes::Label(), SceneOutliner::FColumnInfo(SceneOutliner::EColumnVisibility::Visible, 0));
+		InitOptions.ColumnMap.Add(FSceneOutlinerBuiltInColumnTypes::Label(), FSceneOutlinerColumnInfo(ESceneOutlinerColumnVisibility::Visible, 0));
 
 		MenuContent =
 			SNew(SBox)
diff --git a/Engine/Source/Editor/PropertyEditor/Private/UserInterface/PropertyEditor/SPropertyMenuComponentPicker.h b/Engine/Source/Editor/PropertyEditor/Private/UserInterface/PropertyEditor/SPropertyMenuComponentPicker.h
--- a/Engine/Source/Editor/PropertyEditor/Private/UserInterface/PropertyEditor/SPropertyMenuComponentPicker.h
+++ b/Engine/Source/Editor/PropertyEditor/Private/UserInterface/PropertyEditor/SPropertyMenuComponentPicker.h
@@ -11,7 +11,7 @@ class UActorComponent;
 
 namespace SceneOutliner
 {
-	struct ITreeItem;
+	struct ISceneOutlinerTreeItem;
 }
 
 class SPropertyMenuComponentPicker : public SCompoundWidget
diff --git a/Engine/Source/Editor/PropertyEditor/Private/UserInterface/PropertyEditor/SPropertySceneOutliner.cpp b/Engine/Source/Editor/PropertyEditor/Private/UserInterface/PropertyEditor/SPropertySceneOutliner.cpp
--- a/Engine/Source/Editor/PropertyEditor/Private/UserInterface/PropertyEditor/SPropertySceneOutliner.cpp
+++ b/Engine/Source/Editor/PropertyEditor/Private/UserInterface/PropertyEditor/SPropertySceneOutliner.cpp
@@ -55,7 +55,7 @@ TSharedRef<SWidget> SPropertySceneOutliner::OnGenerateSceneOutliner()
 {
 	FSceneOutlinerModule& SceneOutlinerModule = FModuleManager::Get().LoadModuleChecked<FSceneOutlinerModule>(TEXT("SceneOutliner"));
 
-	SceneOutliner::FInitializationOptions InitOptions;
+	FSceneOutlinerInitializationOptions InitOptions;
 	OnGetActorFilters.ExecuteIfBound( InitOptions.Filters );
 
 	TSharedRef<SWidget> MenuContent = 
diff --git a/Engine/Source/Editor/PropertyEditor/Public/PropertyCustomizationHelpers.h b/Engine/Source/Editor/PropertyEditor/Public/PropertyCustomizationHelpers.h
--- a/Engine/Source/Editor/PropertyEditor/Public/PropertyCustomizationHelpers.h
+++ b/Engine/Source/Editor/PropertyEditor/Public/PropertyCustomizationHelpers.h
@@ -35,18 +35,14 @@ class SToolTip;
 class IPropertyHandle;
 class IDetailGroup;
 class IDetailCategoryBuilder;
-
-namespace SceneOutliner
-{
-	struct FOutlinerFilters;
-}
+struct FSceneOutlinerFilters;
 
 DECLARE_DELEGATE_OneParam(FOnAssetSelected, const FAssetData& /*AssetData*/);
 DECLARE_DELEGATE_RetVal_OneParam(bool, FOnShouldSetAsset, const FAssetData& /*AssetData*/);
 DECLARE_DELEGATE_RetVal_OneParam(bool, FOnShouldFilterAsset, const FAssetData& /*AssetData*/);
 DECLARE_DELEGATE_OneParam(FOnComponentSelected, const UActorComponent* /*ActorComponent*/);
 DECLARE_DELEGATE_RetVal_OneParam(bool, FOnShouldFilterComponent, const UActorComponent* /*ActorComponent*/);
-DECLARE_DELEGATE_OneParam( FOnGetActorFilters, TSharedPtr<SceneOutliner::FOutlinerFilters>& );
+DECLARE_DELEGATE_OneParam( FOnGetActorFilters, TSharedPtr<FSceneOutlinerFilters>& );
 DECLARE_DELEGATE_ThreeParams(FOnGetPropertyComboBoxStrings, TArray< TSharedPtr<FString> >&, TArray<TSharedPtr<SToolTip>>&, TArray<bool>&);
 DECLARE_DELEGATE_RetVal(FString, FOnGetPropertyComboBoxValue);
 DECLARE_DELEGATE_OneParam(FOnPropertyComboBoxValueSelected, const FString&);
