Fixed compile error in DataprepEditorOutlinerMode

diff --git a/Engine/Plugins/Enterprise/DataprepEditor/Source/DataprepEditor/Private/DataprepEditorOutlinerMode.h b/Engine/Plugins/Enterprise/DataprepEditor/Source/DataprepEditor/Private/DataprepEditorOutlinerMode.h
--- a/Engine/Plugins/Enterprise/DataprepEditor/Source/DataprepEditor/Private/DataprepEditorOutlinerMode.h
+++ b/Engine/Plugins/Enterprise/DataprepEditor/Source/DataprepEditor/Private/DataprepEditorOutlinerMode.h
@@ -46,7 +46,7 @@ class FDataprepEditorOutlinerMode : public FActorMode
 {
 public:
 	FDataprepEditorOutlinerMode(SSceneOutliner* InSceneOutliner, TWeakPtr<FDataprepEditor> InDataprepEditor, TWeakObjectPtr<UWorld> InSpecifiedWorldToDisplay = nullptr)
-		: FActorMode(InSceneOutliner, true, InSpecifiedWorldToDisplay)
+		: FActorMode(FActorModeParams(InSceneOutliner, InSpecifiedWorldToDisplay, true, true))
 		, DataprepEditorPtr(InDataprepEditor)
 	{}
 
