Moving FDetailColumnSizeData to PropertyCustomizationHelpers to allow access by other modules. Removed duplicated versions of the struct.

diff --git a/Engine/Plugins/Enterprise/DataprepEditor/Source/DataprepEditor/Private/Widgets/DataprepAssetView.cpp b/Engine/Plugins/Enterprise/DataprepEditor/Source/DataprepEditor/Private/Widgets/DataprepAssetView.cpp
--- a/Engine/Plugins/Enterprise/DataprepEditor/Source/DataprepEditor/Private/Widgets/DataprepAssetView.cpp
+++ b/Engine/Plugins/Enterprise/DataprepEditor/Source/DataprepEditor/Private/Widgets/DataprepAssetView.cpp
@@ -34,7 +34,7 @@
 #include "Widgets/Layout/SSpacer.h"
 #include "Widgets/Text/SInlineEditableTextBlock.h"
 #include "Widgets/Text/STextBlock.h"
-
+#include "PropertyCustomizationHelpers.h"
 
 #define LOCTEXT_NAMESPACE "DataprepAssetView"
 
@@ -404,7 +404,7 @@ void SDataprepAssetView::Construct( const FArguments& InArgs, UDataprepAssetInte
 	bIsChecked = true;
 
 	ColumnWidth = 0.7f;
-	ColumnSizeData = MakeShared< FDataprepDetailsViewColumnSizeData >();
+	ColumnSizeData = MakeShared< FDetailColumnSizeData >();
 	ColumnSizeData->LeftColumnWidth = TAttribute<float>(this, &SDataprepAssetView::OnGetLeftColumnWidth);
 	ColumnSizeData->RightColumnWidth = TAttribute<float>(this, &SDataprepAssetView::OnGetRightColumnWidth);
 	ColumnSizeData->OnWidthChanged = SSplitter::FOnSlotResized::CreateSP(this, &SDataprepAssetView::OnSetColumnWidth);
diff --git a/Engine/Plugins/Enterprise/DataprepEditor/Source/DataprepEditor/Private/Widgets/DataprepAssetView.h b/Engine/Plugins/Enterprise/DataprepEditor/Source/DataprepEditor/Private/Widgets/DataprepAssetView.h
--- a/Engine/Plugins/Enterprise/DataprepEditor/Source/DataprepEditor/Private/Widgets/DataprepAssetView.h
+++ b/Engine/Plugins/Enterprise/DataprepEditor/Source/DataprepEditor/Private/Widgets/DataprepAssetView.h
@@ -9,7 +9,7 @@
 #include "DataprepContentProducer.h"
 #include "DataprepEditor.h"
 #include "DataprepWidgets.h"
-
+#include "PropertyCustomizationHelpers.h"
 #include "Delegates/IDelegateInstance.h"
 #include "Widgets/Input/STextComboBox.h"
 #include "Widgets/SCompoundWidget.h"
@@ -60,7 +60,7 @@ private:
 	TSharedPtr< SDataprepConsumerWidget > ConsumerWidget;
 
 	/** Container used by all splitters in the details view, so that they move in sync */
-	TSharedPtr< FDataprepDetailsViewColumnSizeData > ColumnSizeData;
+	TSharedPtr< FDetailColumnSizeData > ColumnSizeData;
 
 	/** Relative width to control splitters */
 	float ColumnWidth;
diff --git a/Engine/Plugins/Enterprise/DataprepEditor/Source/DataprepEditor/Private/Widgets/DataprepWidgets.cpp b/Engine/Plugins/Enterprise/DataprepEditor/Source/DataprepEditor/Private/Widgets/DataprepWidgets.cpp
--- a/Engine/Plugins/Enterprise/DataprepEditor/Source/DataprepEditor/Private/Widgets/DataprepWidgets.cpp
+++ b/Engine/Plugins/Enterprise/DataprepEditor/Source/DataprepEditor/Private/Widgets/DataprepWidgets.cpp
@@ -49,7 +49,7 @@
 #include "Widgets/Text/SInlineEditableTextBlock.h"
 #include "Widgets/Text/STextBlock.h"
 #include "UObject/UObjectGlobals.h"
-
+#include "PropertyCustomizationHelpers.h"
 #define LOCTEXT_NAMESPACE "DataprepSlateHelper"
 
 namespace DataprepWidgetUtils
@@ -60,7 +60,7 @@ namespace DataprepWidgetUtils
 		SLATE_BEGIN_ARGS(SCustomSplitter) {}
 			SLATE_ARGUMENT( TSharedPtr<SWidget>, NameWidget )
 			SLATE_ARGUMENT( TSharedPtr<SWidget>, ValueWidget )
-			SLATE_ARGUMENT( TSharedPtr< FDataprepDetailsViewColumnSizeData >, ColumnSizeData )
+			SLATE_ARGUMENT( TSharedPtr< FDetailColumnSizeData >, ColumnSizeData )
 		SLATE_END_ARGS();
 
 		void Construct(const FArguments& InArgs)
@@ -121,10 +121,10 @@ namespace DataprepWidgetUtils
 		}
 
 	private:
-		TSharedPtr< FDataprepDetailsViewColumnSizeData > ColumnSizeData;
+		TSharedPtr< FDetailColumnSizeData > ColumnSizeData;
 	};
 
-	TSharedRef<SWidget> CreatePropertyWidget( TSharedPtr<SWidget> NameWidget, TSharedPtr<SWidget> ValueWidget, TSharedPtr< FDataprepDetailsViewColumnSizeData > ColumnSizeData, float Spacing, bool bResizableColumn = true)
+	TSharedRef<SWidget> CreatePropertyWidget( TSharedPtr<SWidget> NameWidget, TSharedPtr<SWidget> ValueWidget, TSharedPtr< FDetailColumnSizeData > ColumnSizeData, float Spacing, bool bResizableColumn = true)
 	{
 		if (bResizableColumn)
 		{
@@ -602,7 +602,7 @@ void SDataprepDetailsView::Construct(const FArguments& InArgs)
 	else
 	{
 		ColumnWidth = 0.7f;
-		ColumnSizeData = MakeShared<FDataprepDetailsViewColumnSizeData>();
+		ColumnSizeData = MakeShared<FDetailColumnSizeData>();
 		ColumnSizeData->LeftColumnWidth = TAttribute<float>(this, &SDataprepDetailsView::OnGetLeftColumnWidth);
 		ColumnSizeData->RightColumnWidth = TAttribute<float>(this, &SDataprepDetailsView::OnGetRightColumnWidth);
 		ColumnSizeData->OnWidthChanged = SSplitter::FOnSlotResized::CreateSP(this, &SDataprepDetailsView::OnSetColumnWidth);
@@ -801,7 +801,7 @@ void SDataprepInstanceParentWidget::Construct(const FArguments& InArgs)
 	else
 	{
 		ColumnWidth = 0.5f;
-		ColumnSizeData = MakeShared<FDataprepDetailsViewColumnSizeData>();
+		ColumnSizeData = MakeShared<FDetailColumnSizeData>();
 		ColumnSizeData->LeftColumnWidth = TAttribute<float>(this, &SDataprepInstanceParentWidget::OnGetLeftColumnWidth);
 		ColumnSizeData->RightColumnWidth = TAttribute<float>(this, &SDataprepInstanceParentWidget::OnGetRightColumnWidth);
 		ColumnSizeData->OnWidthChanged = SSplitter::FOnSlotResized::CreateSP(this, &SDataprepInstanceParentWidget::OnSetColumnWidth);
diff --git a/Engine/Plugins/Enterprise/DataprepEditor/Source/DataprepEditor/Private/Widgets/DataprepWidgets.h b/Engine/Plugins/Enterprise/DataprepEditor/Source/DataprepEditor/Private/Widgets/DataprepWidgets.h
--- a/Engine/Plugins/Enterprise/DataprepEditor/Source/DataprepEditor/Private/Widgets/DataprepWidgets.h
+++ b/Engine/Plugins/Enterprise/DataprepEditor/Source/DataprepEditor/Private/Widgets/DataprepWidgets.h
@@ -52,15 +52,6 @@ namespace DataprepWidgetUtils
 	};
 }
 
-struct FDataprepDetailsViewColumnSizeData
-{
-	TAttribute<float> LeftColumnWidth;
-	TAttribute<float> RightColumnWidth;
-	SSplitter::FOnSlotResized OnWidthChanged;
-
-	void SetColumnWidth(float InWidth) { OnWidthChanged.ExecuteIfBound(InWidth); }
-};
-
 enum class EDataprepCategory 
 { 
 	Producers, 
@@ -75,7 +66,7 @@ class SDataprepCategoryWidget : public STableRow< TSharedPtr< EDataprepCategory
 public:
 
 	SLATE_BEGIN_ARGS(SDataprepCategoryWidget) {}
-	SLATE_ARGUMENT(TSharedPtr< FDataprepDetailsViewColumnSizeData >, ColumnSizeData)
+	SLATE_ARGUMENT(TSharedPtr< FDetailColumnSizeData >, ColumnSizeData)
 	SLATE_ARGUMENT(FText, Title)
 	SLATE_ARGUMENT(TSharedPtr< SWidget >, TitleDetail)
 	SLATE_END_ARGS()
@@ -103,7 +94,7 @@ public:
 		, _ColumnPadding(false)
 		, _ResizableColumn(true)
 		{}
-		SLATE_ARGUMENT(TSharedPtr<FDataprepDetailsViewColumnSizeData>, ColumnSizeData)
+		SLATE_ARGUMENT(TSharedPtr<FDetailColumnSizeData>, ColumnSizeData)
 		SLATE_ARGUMENT(UObject*, Object)
 		SLATE_ARGUMENT(float, Spacing)
 		SLATE_ARGUMENT(bool, ColumnPadding)
@@ -204,7 +195,7 @@ private:
 	float ColumnWidth;
 
 	/** Points to the currently used column size data. Can be provided via argument as well. */
-	TSharedPtr< FDataprepDetailsViewColumnSizeData > ColumnSizeData;
+	TSharedPtr< FDetailColumnSizeData > ColumnSizeData;
 
 	/** If there is a new object to display on the next tick */
 	uint8 bRefreshObjectToDisplay : 1;
@@ -263,7 +254,7 @@ class SDataprepInstanceParentWidget : public SCompoundWidget
 {
 public:
 	SLATE_BEGIN_ARGS(SDataprepInstanceParentWidget) {}
-		SLATE_ARGUMENT(TSharedPtr< FDataprepDetailsViewColumnSizeData >, ColumnSizeData)
+		SLATE_ARGUMENT(TSharedPtr< FDetailColumnSizeData >, ColumnSizeData)
 		SLATE_ARGUMENT(UDataprepAssetInstance*, DataprepInstance)
 	SLATE_END_ARGS()
 
@@ -290,7 +281,7 @@ private:
 	TWeakObjectPtr<UDataprepAssetInstance> DataprepInstancePtr;
 
 	/** Helps sync column resizing with another part of the UI */
-	TSharedPtr< FDataprepDetailsViewColumnSizeData > ColumnSizeData;
+	TSharedPtr< FDetailColumnSizeData > ColumnSizeData;
 
 	/** Relative width to control splitters. */
 	float ColumnWidth;
diff --git a/Engine/Plugins/Enterprise/DataprepEditor/Source/DataprepEditor/Private/Widgets/SDataprepProducersWidget.cpp b/Engine/Plugins/Enterprise/DataprepEditor/Source/DataprepEditor/Private/Widgets/SDataprepProducersWidget.cpp
--- a/Engine/Plugins/Enterprise/DataprepEditor/Source/DataprepEditor/Private/Widgets/SDataprepProducersWidget.cpp
+++ b/Engine/Plugins/Enterprise/DataprepEditor/Source/DataprepEditor/Private/Widgets/SDataprepProducersWidget.cpp
@@ -32,7 +32,7 @@
 #include "Widgets/Layout/SSpacer.h"
 #include "Widgets/Text/SInlineEditableTextBlock.h"
 #include "Widgets/Text/STextBlock.h"
-
+#include "PropertyCustomizationHelpers.h"
 #define LOCTEXT_NAMESPACE "DataprepProducersWidget"
 
 FContentProducerEntry::FContentProducerEntry(int32 InProducerIndex, UDataprepAssetProducers* InAssetProducersPtr)
@@ -76,7 +76,7 @@ void FContentProducerEntry::RemoveProducer()
 
 
 /** Construct function for this widget */
-void SDataprepProducersTableRow::Construct(const FArguments& InArgs, const TSharedRef<STableViewBase>& OwnerTableView, const TSharedRef<FContentProducerEntry>& InNode, TSharedRef< FDataprepDetailsViewColumnSizeData > InColumnSizeData)
+void SDataprepProducersTableRow::Construct(const FArguments& InArgs, const TSharedRef<STableViewBase>& OwnerTableView, const TSharedRef<FContentProducerEntry>& InNode, TSharedRef< FDetailColumnSizeData > InColumnSizeData)
 {
 	Node = InNode;
 	STableRow::Construct(STableRow::FArguments(), OwnerTableView);
@@ -93,7 +93,7 @@ void SDataprepProducersTableRow::Construct(const FArguments& InArgs, const TShar
 	}
 }
 
-TSharedRef<SWidget> SDataprepProducersTableRow::GetInputMainWidget( TSharedRef< FDataprepDetailsViewColumnSizeData > ColumnSizeData )
+TSharedRef<SWidget> SDataprepProducersTableRow::GetInputMainWidget( TSharedRef< FDetailColumnSizeData > ColumnSizeData )
 {
 	FContentProducerEntryPtr ProducerStackEntry = Node.Pin();
 
@@ -147,7 +147,7 @@ TSharedRef<SWidget> SDataprepProducersTableRow::GetInputMainWidget( TSharedRef<
 	return Widget.ToSharedRef();
 }
 
-void SDataprepProducersTreeView::Construct(const FArguments& InArgs, UDataprepAssetProducers* InAssetProducersPtr, TSharedRef< FDataprepDetailsViewColumnSizeData > InColumnSizeData )
+void SDataprepProducersTreeView::Construct(const FArguments& InArgs, UDataprepAssetProducers* InAssetProducersPtr, TSharedRef< FDetailColumnSizeData > InColumnSizeData )
 {
 	ColumnSizeData = InColumnSizeData;
 	AssetProducersPtr = InAssetProducersPtr;
diff --git a/Engine/Plugins/Enterprise/DataprepEditor/Source/DataprepEditor/Private/Widgets/SDataprepProducersWidget.h b/Engine/Plugins/Enterprise/DataprepEditor/Source/DataprepEditor/Private/Widgets/SDataprepProducersWidget.h
--- a/Engine/Plugins/Enterprise/DataprepEditor/Source/DataprepEditor/Private/Widgets/SDataprepProducersWidget.h
+++ b/Engine/Plugins/Enterprise/DataprepEditor/Source/DataprepEditor/Private/Widgets/SDataprepProducersWidget.h
@@ -16,7 +16,7 @@
 class FUICommandList;
 class UDataprepContentProducer;
 class UDataprepAssetProducers;
-struct FDataprepDetailsViewColumnSizeData;
+struct FDetailColumnSizeData;
 
 class FContentProducerEntry
 {
@@ -56,9 +56,9 @@ public:
 	SLATE_BEGIN_ARGS(SDataprepProducersTableRow) {}
 	SLATE_END_ARGS()
 
-	void Construct(const FArguments& InArgs, const TSharedRef<STableViewBase>& OwnerTableView, const TSharedRef<FContentProducerEntry>& InNode, TSharedRef< FDataprepDetailsViewColumnSizeData > InColumnSizeData);
+	void Construct(const FArguments& InArgs, const TSharedRef<STableViewBase>& OwnerTableView, const TSharedRef<FContentProducerEntry>& InNode, TSharedRef< FDetailColumnSizeData > InColumnSizeData);
 
-	TSharedRef<SWidget> GetInputMainWidget( TSharedRef< FDataprepDetailsViewColumnSizeData > ColumnSizeData );
+	TSharedRef<SWidget> GetInputMainWidget( TSharedRef< FDetailColumnSizeData > ColumnSizeData );
 
 	TSharedPtr<FContentProducerEntry> GetDisplayNode() const
 	{
@@ -76,7 +76,7 @@ public:
 	SLATE_BEGIN_ARGS(SDataprepProducersTreeView) {}
 	SLATE_END_ARGS()
 
-	void Construct(const FArguments& InArgs, UDataprepAssetProducers* InAssetProducersPtr, TSharedRef< FDataprepDetailsViewColumnSizeData > InColumnSizeData );
+	void Construct(const FArguments& InArgs, UDataprepAssetProducers* InAssetProducersPtr, TSharedRef< FDetailColumnSizeData > InColumnSizeData );
 
 	int32 GetDisplayIndexOfNode(FContentProducerEntryRef InNode);
 
@@ -94,7 +94,7 @@ private:
 private:
 	TWeakObjectPtr<UDataprepAssetProducers> AssetProducersPtr;
 	TArray<FContentProducerEntryRef> RootNodes;
-	TSharedPtr< FDataprepDetailsViewColumnSizeData > ColumnSizeData;
+	TSharedPtr< FDetailColumnSizeData > ColumnSizeData;
 };
 
 class SDataprepProducersWidget : public SCompoundWidget
@@ -102,7 +102,7 @@ class SDataprepProducersWidget : public SCompoundWidget
 public:
 
 	SLATE_BEGIN_ARGS(SDataprepProducersWidget) {}
-	SLATE_ARGUMENT(TSharedPtr< FDataprepDetailsViewColumnSizeData >, ColumnSizeData)
+	SLATE_ARGUMENT(TSharedPtr< FDetailColumnSizeData >, ColumnSizeData)
 	SLATE_END_ARGS()
 
 	~SDataprepProducersWidget();
diff --git a/Engine/Plugins/Enterprise/VariantManager/Source/VariantManager/Private/DisplayNodes/VariantManagerPropertyNode.cpp b/Engine/Plugins/Enterprise/VariantManager/Source/VariantManager/Private/DisplayNodes/VariantManagerPropertyNode.cpp
--- a/Engine/Plugins/Enterprise/VariantManager/Source/VariantManager/Private/DisplayNodes/VariantManagerPropertyNode.cpp
+++ b/Engine/Plugins/Enterprise/VariantManager/Source/VariantManager/Private/DisplayNodes/VariantManagerPropertyNode.cpp
@@ -26,7 +26,7 @@
 #include "Widgets/DeclarativeSyntaxSupport.h"
 #include "Widgets/Images/SImage.h"
 #include "Widgets/Input/SButton.h"
-
+#include "PropertyCustomizationHelpers.h"
 #define LOCTEXT_NAMESPACE "FVariantManagerPropertyNode"
 
 using FDisplayNodeRef = TSharedRef<FVariantManagerDisplayNode>;
@@ -373,7 +373,7 @@ TSharedPtr<SWidget> FVariantManagerPropertyNode::GetPropertyValueWidget()
 TSharedRef<SWidget> FVariantManagerPropertyNode::GetCustomOutlinerContent(TSharedPtr<SVariantManagerTableRow> InTableRow)
 {
 	// Using this syncs all splitters between property nodes and also the header
-	FPropertyColumnSizeData& ColumnSizeData = GetVariantManager().Pin()->GetVariantManagerWidget()->GetPropertyColumnSizeData();
+	FDetailColumnSizeData& ColumnSizeData = GetVariantManager().Pin()->GetVariantManagerWidget()->GetPropertyColumnSizeData();
 
 	return SNew(SBox)
 	[
diff --git a/Engine/Plugins/Enterprise/VariantManager/Source/VariantManager/Private/SVariantManager.h b/Engine/Plugins/Enterprise/VariantManager/Source/VariantManager/Private/SVariantManager.h
--- a/Engine/Plugins/Enterprise/VariantManager/Source/VariantManager/Private/SVariantManager.h
+++ b/Engine/Plugins/Enterprise/VariantManager/Source/VariantManager/Private/SVariantManager.h
@@ -10,7 +10,7 @@
 #include "Widgets/Views/SListView.h"
 #include "Framework/Commands/UICommandList.h"
 #include "UObject/WeakObjectPtr.h"
-
+#include "PropertyCustomizationHelpers.h"
 #include "VariantManager.h"
 
 class FExtender;
@@ -63,16 +63,6 @@ struct FSplitterValues
 	FString ToString();
 };
 
-// Replica of FDetailColumnSizeData used by DetailViews
-struct FPropertyColumnSizeData
-{
-	TAttribute<float> LeftColumnWidth;
-	TAttribute<float> RightColumnWidth;
-	SSplitter::FOnSlotResized OnWidthChanged;
-
-	void SetColumnWidth(float InWidth) { OnWidthChanged.ExecuteIfBound(InWidth); }
-};
-
 class SVariantManager
 	: public SCompoundWidget
 	, public FNotifyHook
@@ -187,7 +177,7 @@ public:
 	void SortDisplayNodes(TArray<TSharedRef<FVariantManagerDisplayNode>>& DisplayNodes);
 
 	TSharedRef<SWidget> MakeAddButton();
-	FPropertyColumnSizeData& GetPropertyColumnSizeData()
+	FDetailColumnSizeData& GetPropertyColumnSizeData()
 	{
 		return ColumnSizeData;
 	}
@@ -263,7 +253,7 @@ private:
 	bool bAutoCaptureProperties = false;
 
 	// Mirrors detailview, its used by all splitters in the column, so that they move in sync
-	FPropertyColumnSizeData ColumnSizeData;
+	FDetailColumnSizeData ColumnSizeData;
 	float RightPropertyColumnWidth;
 
 	FDelegateHandle OnObjectTransactedHandle;
diff --git a/Engine/Source/Editor/MaterialEditor/Private/SMaterialLayersFunctionsTree.h b/Engine/Source/Editor/MaterialEditor/Private/SMaterialLayersFunctionsTree.h
--- a/Engine/Source/Editor/MaterialEditor/Private/SMaterialLayersFunctionsTree.h
+++ b/Engine/Source/Editor/MaterialEditor/Private/SMaterialLayersFunctionsTree.h
@@ -14,7 +14,7 @@
 #include "IDetailTreeNode.h"
 #include "IDetailPropertyRow.h"
 #include "MaterialPropertyHelpers.h"
-
+#include "PropertyCustomizationHelpers.h"
 class IPropertyHandle;
 class UMaterialEditorInstanceConstant;
 class SMaterialLayersFunctionsInstanceTree;
@@ -34,7 +34,7 @@ public:
 	SLATE_ARGUMENT(SMaterialLayersFunctionsInstanceTree*, InTree)
 	SLATE_END_ARGS()
 
-	FMaterialTreeColumnSizeData ColumnSizeData;
+	FDetailColumnSizeData ColumnSizeData;
 	bool bIsBeingDragged;
 
 private:
@@ -213,7 +213,7 @@ public:
 	SLATE_ARGUMENT(SMaterialLayersFunctionsMaterialTree*, InTree)
 	SLATE_END_ARGS()
 
-	FMaterialTreeColumnSizeData ColumnSizeData;
+	FDetailColumnSizeData ColumnSizeData;
 	bool bIsBeingDragged;
 
 private:
diff --git a/Engine/Source/Editor/MaterialEditor/Private/SMaterialParametersOverviewWidget.h b/Engine/Source/Editor/MaterialEditor/Private/SMaterialParametersOverviewWidget.h
--- a/Engine/Source/Editor/MaterialEditor/Private/SMaterialParametersOverviewWidget.h
+++ b/Engine/Source/Editor/MaterialEditor/Private/SMaterialParametersOverviewWidget.h
@@ -13,7 +13,7 @@
 #include "IDetailTreeNode.h"
 #include "IDetailPropertyRow.h"
 #include "MaterialPropertyHelpers.h"
-
+#include "PropertyCustomizationHelpers.h"
 class IPropertyHandle;
 class SMaterialParametersOverviewTree;
 class UMaterialEditorPreviewParameters;
@@ -59,7 +59,7 @@ private:
 	/** The set of material parameters this is associated with */
 	UMaterialEditorPreviewParameters* MaterialEditorInstance;
 
-	FMaterialTreeColumnSizeData ColumnSizeData;
+	FDetailColumnSizeData ColumnSizeData;
 };
 
 // ********* SMaterialParametersOverviewPanel *******
diff --git a/Engine/Source/Editor/MaterialEditor/Public/MaterialPropertyHelpers.h b/Engine/Source/Editor/MaterialEditor/Public/MaterialPropertyHelpers.h
--- a/Engine/Source/Editor/MaterialEditor/Public/MaterialPropertyHelpers.h
+++ b/Engine/Source/Editor/MaterialEditor/Public/MaterialPropertyHelpers.h
@@ -71,15 +71,6 @@ struct FUnsortedParamData
 	TSharedPtr<class IPropertyHandle> ParameterHandle;
 };
 
-struct FMaterialTreeColumnSizeData
-{
-	TAttribute<float> LeftColumnWidth;
-	TAttribute<float> RightColumnWidth;
-	SSplitter::FOnSlotResized OnWidthChanged;
-
-	void SetColumnWidth(float InWidth) { OnWidthChanged.ExecuteIfBound(InWidth); }
-};
-
 class SLayerHandle : public SCompoundWidget
 {
 public:
diff --git a/Engine/Source/Editor/PropertyEditor/Private/DetailAdvancedDropdownNode.cpp b/Engine/Source/Editor/PropertyEditor/Private/DetailAdvancedDropdownNode.cpp
--- a/Engine/Source/Editor/PropertyEditor/Private/DetailAdvancedDropdownNode.cpp
+++ b/Engine/Source/Editor/PropertyEditor/Private/DetailAdvancedDropdownNode.cpp
@@ -4,6 +4,7 @@
 #include "Widgets/Images/SImage.h"
 #include "Widgets/Input/SButton.h"
 #include "SDetailTableRowBase.h"
+#include "PropertyCustomizationHelpers.h"
 
 class SAdvancedDropdownRow : public SDetailTableRowBase
 {
diff --git a/Engine/Source/Editor/PropertyEditor/Private/DetailAdvancedDropdownNode.h b/Engine/Source/Editor/PropertyEditor/Private/DetailAdvancedDropdownNode.h
--- a/Engine/Source/Editor/PropertyEditor/Private/DetailAdvancedDropdownNode.h
+++ b/Engine/Source/Editor/PropertyEditor/Private/DetailAdvancedDropdownNode.h
@@ -10,6 +10,7 @@
 #include "Widgets/Views/STableViewBase.h"
 #include "Widgets/Views/STableRow.h"
 #include "DetailCategoryBuilderImpl.h"
+#include "PropertyCustomizationHelpers.h"
 
 class FAdvancedDropdownNode : public FDetailTreeNode, public TSharedFromThis<FAdvancedDropdownNode>
 {
diff --git a/Engine/Source/Editor/PropertyEditor/Private/DetailCategoryBuilderImpl.cpp b/Engine/Source/Editor/PropertyEditor/Private/DetailCategoryBuilderImpl.cpp
--- a/Engine/Source/Editor/PropertyEditor/Private/DetailCategoryBuilderImpl.cpp
+++ b/Engine/Source/Editor/PropertyEditor/Private/DetailCategoryBuilderImpl.cpp
@@ -14,6 +14,7 @@
 #include "ItemPropertyNode.h"
 #include "IPropertyGenerationUtilities.h"
 #include "DetailBuilderTypes.h"
+#include "PropertyCustomizationHelpers.h"
 
 namespace DetailLayoutConstants
 {
diff --git a/Engine/Source/Editor/PropertyEditor/Private/DetailCategoryBuilderImpl.h b/Engine/Source/Editor/PropertyEditor/Private/DetailCategoryBuilderImpl.h
--- a/Engine/Source/Editor/PropertyEditor/Private/DetailCategoryBuilderImpl.h
+++ b/Engine/Source/Editor/PropertyEditor/Private/DetailCategoryBuilderImpl.h
@@ -17,6 +17,7 @@
 #include "DetailLayoutBuilderImpl.h"
 #include "IDetailCustomNodeBuilder.h"
 #include "DetailCategoryBuilder.h"
+#include "PropertyCustomizationHelpers.h"
 
 class FDetailGroup;
 class FDetailPropertyRow;
diff --git a/Engine/Source/Editor/PropertyEditor/Private/DetailCategoryGroupNode.cpp b/Engine/Source/Editor/PropertyEditor/Private/DetailCategoryGroupNode.cpp
--- a/Engine/Source/Editor/PropertyEditor/Private/DetailCategoryGroupNode.cpp
+++ b/Engine/Source/Editor/PropertyEditor/Private/DetailCategoryGroupNode.cpp
@@ -1,7 +1,7 @@
 // Copyright Epic Games, Inc. All Rights Reserved.
 
 #include "DetailCategoryGroupNode.h"
-
+#include "PropertyCustomizationHelpers.h"
 void SDetailCategoryTableRow::Construct( const FArguments& InArgs, TSharedRef<FDetailTreeNode> InOwnerTreeNode, const TSharedRef<STableViewBase>& InOwnerTableView )
 {
 	OwnerTreeNode = InOwnerTreeNode;
diff --git a/Engine/Source/Editor/PropertyEditor/Private/DetailCategoryGroupNode.h b/Engine/Source/Editor/PropertyEditor/Private/DetailCategoryGroupNode.h
--- a/Engine/Source/Editor/PropertyEditor/Private/DetailCategoryGroupNode.h
+++ b/Engine/Source/Editor/PropertyEditor/Private/DetailCategoryGroupNode.h
@@ -12,6 +12,7 @@
 #include "Widgets/Views/STableRow.h"
 #include "DetailCategoryBuilderImpl.h"
 #include "SDetailTableRowBase.h"
+#include "PropertyCustomizationHelpers.h"
 
 class SDetailCategoryTableRow : public SDetailTableRowBase
 {
diff --git a/Engine/Source/Editor/PropertyEditor/Private/DetailItemNode.cpp b/Engine/Source/Editor/PropertyEditor/Private/DetailItemNode.cpp
--- a/Engine/Source/Editor/PropertyEditor/Private/DetailItemNode.cpp
+++ b/Engine/Source/Editor/PropertyEditor/Private/DetailItemNode.cpp
@@ -7,7 +7,7 @@
 #include "SDetailSingleItemRow.h"
 #include "IDetailKeyframeHandler.h"
 #include "ObjectPropertyNode.h"
-
+#include "PropertyCustomizationHelpers.h"
 
 FDetailItemNode::FDetailItemNode(const FDetailLayoutCustomization& InCustomization, TSharedRef<FDetailCategoryImpl> InParentCategory, TAttribute<bool> InIsParentEnabled, TSharedPtr<IDetailGroup> InParentGroup)
 	: Customization( InCustomization )
diff --git a/Engine/Source/Editor/PropertyEditor/Private/DetailItemNode.h b/Engine/Source/Editor/PropertyEditor/Private/DetailItemNode.h
--- a/Engine/Source/Editor/PropertyEditor/Private/DetailItemNode.h
+++ b/Engine/Source/Editor/PropertyEditor/Private/DetailItemNode.h
@@ -10,6 +10,7 @@
 #include "Widgets/Views/STableViewBase.h"
 #include "Widgets/Views/STableRow.h"
 #include "DetailCategoryBuilderImpl.h"
+#include "PropertyCustomizationHelpers.h"
 
 /**
  * A single item in a detail tree                                                              
diff --git a/Engine/Source/Editor/PropertyEditor/Private/DetailMultiTopLevelObjectRootNode.cpp b/Engine/Source/Editor/PropertyEditor/Private/DetailMultiTopLevelObjectRootNode.cpp
--- a/Engine/Source/Editor/PropertyEditor/Private/DetailMultiTopLevelObjectRootNode.cpp
+++ b/Engine/Source/Editor/PropertyEditor/Private/DetailMultiTopLevelObjectRootNode.cpp
@@ -5,7 +5,7 @@
 #include "DetailWidgetRow.h"
 #include "ObjectPropertyNode.h"
 #include "Misc/ConfigCacheIni.h"
-
+#include "PropertyCustomizationHelpers.h"
 
 void SDetailMultiTopLevelObjectTableRow::Construct( const FArguments& InArgs, TSharedRef<FDetailTreeNode> InOwnerTreeNode, const TSharedRef<STableViewBase>& InOwnerTableView )
 {
diff --git a/Engine/Source/Editor/PropertyEditor/Private/DetailMultiTopLevelObjectRootNode.h b/Engine/Source/Editor/PropertyEditor/Private/DetailMultiTopLevelObjectRootNode.h
--- a/Engine/Source/Editor/PropertyEditor/Private/DetailMultiTopLevelObjectRootNode.h
+++ b/Engine/Source/Editor/PropertyEditor/Private/DetailMultiTopLevelObjectRootNode.h
@@ -13,6 +13,7 @@
 #include "SDetailsViewBase.h"
 #include "SDetailTableRowBase.h"
 #include "IDetailRootObjectCustomization.h"
+#include "PropertyCustomizationHelpers.h"
 
 using EExpansionArrowUsage = IDetailRootObjectCustomization::EExpansionArrowUsage;
 
diff --git a/Engine/Source/Editor/PropertyEditor/Private/DetailTreeNode.h b/Engine/Source/Editor/PropertyEditor/Private/DetailTreeNode.h
--- a/Engine/Source/Editor/PropertyEditor/Private/DetailTreeNode.h
+++ b/Engine/Source/Editor/PropertyEditor/Private/DetailTreeNode.h
@@ -7,6 +7,7 @@
 #include "PropertyPath.h"
 #include "IPropertyUtilities.h"
 #include "IDetailTreeNode.h"
+#include "PropertyCustomizationHelpers.h"
 
 class FDetailCategoryImpl;
 class IDetailsViewPrivate;
diff --git a/Engine/Source/Editor/PropertyEditor/Private/SDetailSingleItemRow.h b/Engine/Source/Editor/PropertyEditor/Private/SDetailSingleItemRow.h
--- a/Engine/Source/Editor/PropertyEditor/Private/SDetailSingleItemRow.h
+++ b/Engine/Source/Editor/PropertyEditor/Private/SDetailSingleItemRow.h
@@ -12,6 +12,7 @@
 #include "SDetailTableRowBase.h"
 #include "DragAndDrop/DecoratedDragDropOp.h"
 #include "ScopedTransaction.h"
+#include "PropertyCustomizationHelpers.h"
 
 class IDetailKeyframeHandler;
 struct FDetailLayoutCustomization;
diff --git a/Engine/Source/Editor/PropertyEditor/Private/SDetailsViewBase.h b/Engine/Source/Editor/PropertyEditor/Private/SDetailsViewBase.h
--- a/Engine/Source/Editor/PropertyEditor/Private/SDetailsViewBase.h
+++ b/Engine/Source/Editor/PropertyEditor/Private/SDetailsViewBase.h
@@ -21,6 +21,7 @@
 #include "IDetailsView.h"
 #include "IDetailsViewPrivate.h"
 #include "PropertyRowGenerator.h"
+#include "PropertyCustomizationHelpers.h"
 
 class FDetailCategoryImpl;
 class FDetailLayoutBuilderImpl;
@@ -69,14 +70,7 @@ struct FDetailFilter
 	TSet<FPropertyPath> WhitelistedProperties;
 };
 
-struct FDetailColumnSizeData
-{
-	TAttribute<float> LeftColumnWidth;
-	TAttribute<float> RightColumnWidth;
-	SSplitter::FOnSlotResized OnWidthChanged;
 
-	void SetColumnWidth(float InWidth) { OnWidthChanged.ExecuteIfBound(InWidth); }
-};
 
 class SDetailsViewBase : public IDetailsViewPrivate
 {
diff --git a/Engine/Source/Editor/PropertyEditor/Public/PropertyCustomizationHelpers.h b/Engine/Source/Editor/PropertyEditor/Public/PropertyCustomizationHelpers.h
--- a/Engine/Source/Editor/PropertyEditor/Public/PropertyCustomizationHelpers.h
+++ b/Engine/Source/Editor/PropertyEditor/Public/PropertyCustomizationHelpers.h
@@ -787,3 +787,12 @@ private:
 	int32 ThumbnailSize;
 	int32 SectionsLodIndex;
 };
+
+struct FDetailColumnSizeData
+{
+	TAttribute<float> LeftColumnWidth;
+	TAttribute<float> RightColumnWidth;
+	SSplitter::FOnSlotResized OnWidthChanged;
+
+	void SetColumnWidth(float InWidth) { OnWidthChanged.ExecuteIfBound(InWidth); }
+};
\ No newline at end of file
