Adding missing forward declaration

diff --git a/Engine/Plugins/Enterprise/DataprepEditor/Source/DataprepEditor/Private/Widgets/DataprepWidgets.h b/Engine/Plugins/Enterprise/DataprepEditor/Source/DataprepEditor/Private/Widgets/DataprepWidgets.h
--- a/Engine/Plugins/Enterprise/DataprepEditor/Source/DataprepEditor/Private/Widgets/DataprepWidgets.h
+++ b/Engine/Plugins/Enterprise/DataprepEditor/Source/DataprepEditor/Private/Widgets/DataprepWidgets.h
@@ -21,6 +21,7 @@ class UDataprepContentConsumer;
 class UDataprepParameterizableObject;
 
 struct FAssetData;
+struct FDetailColumnSizeData;
 
 namespace DataprepWidgetUtils
 {
