Fixed compilation issues caused by CL14375859 in plugins that I accidentally wasn't compiling.

diff --git a/Engine/Plugins/Enterprise/DataprepEditor/Source/DataprepEditor/Private/Widgets/DataprepWidgets.cpp b/Engine/Plugins/Enterprise/DataprepEditor/Source/DataprepEditor/Private/Widgets/DataprepWidgets.cpp
--- a/Engine/Plugins/Enterprise/DataprepEditor/Source/DataprepEditor/Private/Widgets/DataprepWidgets.cpp
+++ b/Engine/Plugins/Enterprise/DataprepEditor/Source/DataprepEditor/Private/Widgets/DataprepWidgets.cpp
@@ -76,15 +76,15 @@ namespace DataprepWidgetUtils
 			ColumnSizeData = InArgs._ColumnSizeData;
 
 			AddSlot()
-			.Value(ColumnSizeData->LeftColumnWidth)
-			.OnSlotResized( SSplitter::FOnSlotResized::CreateLambda( [](float InNewWidth) -> void {} ) )
+			.Value(ColumnSizeData->NameColumnWidth)
+			.OnSlotResized(ColumnSizeData->OnNameColumnResized)
 			[
 				InArgs._NameWidget.ToSharedRef()
 			];
 
 			AddSlot()
-			.Value(InArgs._ColumnSizeData->RightColumnWidth)
-			.OnSlotResized(ColumnSizeData->OnWidthChanged)
+			.Value(InArgs._ColumnSizeData->ValueColumnWidth)
+			.OnSlotResized(ColumnSizeData->OnValueColumnResized)
 			[
 				SNew(SHorizontalBox)
 				+ SHorizontalBox::Slot()
@@ -107,14 +107,14 @@ namespace DataprepWidgetUtils
 				const FVector2D AllottedSize = AllottedGeometry.GetLocalSize();
 
 				FVector2D LocalPosition = FVector2D::ZeroVector;
-				FVector2D LocalSize(  AllottedSize.X * ColumnSizeData->LeftColumnWidth.Get(0.f), AllottedSize.Y );
+				FVector2D LocalSize(  AllottedSize.X * ColumnSizeData->NameColumnWidth.Get(0.f), AllottedSize.Y );
 
 				ArrangedChildren.AddWidget( EVisibility::Visible, AllottedGeometry.MakeChild( LeftChild.GetWidget(), LocalPosition, LocalSize ));
 
 				const SSplitter::FSlot& RightChild = Children[1];
 
 				LocalPosition = FVector2D(LocalSize.X, 0.f);
-				LocalSize.X = AllottedSize.X * ColumnSizeData->RightColumnWidth.Get(0.f);
+				LocalSize.X = AllottedSize.X * ColumnSizeData->ValueColumnWidth.Get(0.f);
 
 				ArrangedChildren.AddWidget( EVisibility::Visible, AllottedGeometry.MakeChild( RightChild.GetWidget(), LocalPosition, LocalSize ));
 			}
@@ -595,14 +595,21 @@ void SDataprepDetailsView::Construct(const FArguments& InArgs)
 	bColumnPadding = InArgs._ColumnPadding;
 	bResizableColumn = InArgs._ResizableColumn;
 
-	ColumnSizeData = MakeShared<FDetailColumnSizeData>();
+	if ( InArgs._ColumnSizeData.IsValid() )
+	{
+		ColumnSizeData = InArgs._ColumnSizeData;
+	}
+	else
+	{
+		ColumnSizeData = MakeShared<FDetailColumnSizeData>();
+	}
 
 	FPropertyEditorModule& PropertyEditorModule = FModuleManager::GetModuleChecked<FPropertyEditorModule>("PropertyEditor");
 
 	FPropertyRowGeneratorArgs Args;
 	Generator = PropertyEditorModule.CreatePropertyRowGenerator(Args);
 
-	if( DetailedObject != nullptr )
+	if ( DetailedObject != nullptr )
 	{
 		TArray< UObject* > Objects;
 		Objects.Add( DetailedObject );
@@ -783,6 +790,8 @@ void SDataprepInstanceParentWidget::Construct(const FArguments& InArgs)
 		return;
 	}
 
+	ColumnSizeData = InArgs._ColumnSizeData;
+
 	TSharedRef<SWidget> NameWidget = SNew(SHorizontalBox)
 	.Clipping(EWidgetClipping::OnDemand)
 	+ SHorizontalBox::Slot()
diff --git a/Engine/Plugins/Enterprise/DataprepEditor/Source/DataprepEditor/Private/Widgets/DataprepWidgets.h b/Engine/Plugins/Enterprise/DataprepEditor/Source/DataprepEditor/Private/Widgets/DataprepWidgets.h
--- a/Engine/Plugins/Enterprise/DataprepEditor/Source/DataprepEditor/Private/Widgets/DataprepWidgets.h
+++ b/Engine/Plugins/Enterprise/DataprepEditor/Source/DataprepEditor/Private/Widgets/DataprepWidgets.h
@@ -21,7 +21,7 @@ class UDataprepContentConsumer;
 class UDataprepParameterizableObject;
 
 struct FAssetData;
-struct FDetailColumnSizeData;
+class FDetailColumnSizeData;
 
 namespace DataprepWidgetUtils
 {
@@ -67,8 +67,9 @@ class SDataprepCategoryWidget : public STableRow< TSharedPtr< EDataprepCategory
 public:
 
 	SLATE_BEGIN_ARGS(SDataprepCategoryWidget) {}
-	SLATE_ARGUMENT(FText, Title)
-	SLATE_ARGUMENT(TSharedPtr< SWidget >, TitleDetail)
+		SLATE_ARGUMENT(FText, Title)
+		SLATE_ARGUMENT(TSharedPtr< SWidget >, TitleDetail)
+		SLATE_ARGUMENT(TSharedPtr< FDetailColumnSizeData >, ColumnSizeData)
 	SLATE_END_ARGS()
 
 	void Construct( const FArguments& InArgs, TSharedRef< SWidget > InContent, const TSharedRef<STableViewBase>& InOwnerTableView );
@@ -98,6 +99,7 @@ public:
 		SLATE_ARGUMENT(float, Spacing)
 		SLATE_ARGUMENT(bool, ColumnPadding)
 		SLATE_ARGUMENT(bool, ResizableColumn)
+		SLATE_ARGUMENT(TSharedPtr< FDetailColumnSizeData >, ColumnSizeData)
 	SLATE_END_ARGS()
 
 public:
diff --git a/Engine/Plugins/Enterprise/DataprepEditor/Source/DataprepEditor/Private/Widgets/SDataprepProducersWidget.h b/Engine/Plugins/Enterprise/DataprepEditor/Source/DataprepEditor/Private/Widgets/SDataprepProducersWidget.h
--- a/Engine/Plugins/Enterprise/DataprepEditor/Source/DataprepEditor/Private/Widgets/SDataprepProducersWidget.h
+++ b/Engine/Plugins/Enterprise/DataprepEditor/Source/DataprepEditor/Private/Widgets/SDataprepProducersWidget.h
@@ -16,7 +16,7 @@
 class FUICommandList;
 class UDataprepContentProducer;
 class UDataprepAssetProducers;
-struct FDetailColumnSizeData;
+class FDetailColumnSizeData;
 
 class FContentProducerEntry
 {
@@ -102,7 +102,7 @@ class SDataprepProducersWidget : public SCompoundWidget
 public:
 
 	SLATE_BEGIN_ARGS(SDataprepProducersWidget) {}
-	SLATE_ARGUMENT(TSharedPtr< FDetailColumnSizeData >, ColumnSizeData)
+		SLATE_ARGUMENT(TSharedPtr< FDetailColumnSizeData >, ColumnSizeData)
 	SLATE_END_ARGS()
 
 	~SDataprepProducersWidget();
diff --git a/Engine/Plugins/Enterprise/VariantManager/Source/VariantManager/Private/DisplayNodes/VariantManagerPropertyNode.cpp b/Engine/Plugins/Enterprise/VariantManager/Source/VariantManager/Private/DisplayNodes/VariantManagerPropertyNode.cpp
--- a/Engine/Plugins/Enterprise/VariantManager/Source/VariantManager/Private/DisplayNodes/VariantManagerPropertyNode.cpp
+++ b/Engine/Plugins/Enterprise/VariantManager/Source/VariantManager/Private/DisplayNodes/VariantManagerPropertyNode.cpp
@@ -390,12 +390,8 @@ TSharedRef<SWidget> FVariantManagerPropertyNode::GetCustomOutlinerContent(TShare
 			.Orientation(Orient_Horizontal)
 
 			+ SSplitter::Slot()
-			.Value(ColumnSizeData.LeftColumnWidth)
-			.OnSlotResized(SSplitter::FOnSlotResized::CreateLambda([](float InNewWidth)
-			{
-				 //This has to be bound or the splitter will take it upon itself to determine the size
-				 //We do nothing here because it is handled by the column size data
-			}))
+			.Value(ColumnSizeData.NameColumnWidth)
+			.OnSlotResized(ColumnSizeData.OnNameColumnResized)
 			[
 				SNew(SBox)
 				.VAlign(VAlign_Center)
@@ -415,8 +411,8 @@ TSharedRef<SWidget> FVariantManagerPropertyNode::GetCustomOutlinerContent(TShare
 			]
 
 			+ SSplitter::Slot()
-			.Value(ColumnSizeData.RightColumnWidth)
-			.OnSlotResized(ColumnSizeData.OnWidthChanged)
+			.Value(ColumnSizeData.ValueColumnWidth)
+			.OnSlotResized(ColumnSizeData.OnValueColumnResized)
 			[
 				SNew(SHorizontalBox)
 				+SHorizontalBox::Slot()
diff --git a/Engine/Plugins/Enterprise/VariantManager/Source/VariantManager/Private/SVariantManager.cpp b/Engine/Plugins/Enterprise/VariantManager/Source/VariantManager/Private/SVariantManager.cpp
--- a/Engine/Plugins/Enterprise/VariantManager/Source/VariantManager/Private/SVariantManager.cpp
+++ b/Engine/Plugins/Enterprise/VariantManager/Source/VariantManager/Private/SVariantManager.cpp
@@ -176,9 +176,6 @@ void SVariantManager::Construct(const FArguments& InArgs, TSharedRef<FVariantMan
 	}
 
 	RightPropertyColumnWidth = SplitterValues.PropertyValueColumn / (SplitterValues.PropertyValueColumn + SplitterValues.PropertyNameColumn);
-	ColumnSizeData.LeftColumnWidth = TAttribute<float>(this, &SVariantManager::OnGetLeftColumnWidth);
-	ColumnSizeData.RightColumnWidth = TAttribute<float>(this, &SVariantManager::OnGetRightColumnWidth);
-	ColumnSizeData.OnWidthChanged = SSplitter::FOnSlotResized::CreateSP(this, &SVariantManager::OnSetColumnWidth);
 
 	InVariantManager->GetSelection().GetOnOutlinerNodeSelectionChanged().AddSP(this, &SVariantManager::RefreshActorList);
 	InVariantManager->GetSelection().GetOnActorNodeSelectionChanged().AddSP(this, &SVariantManager::OnActorNodeSelectionChanged);
@@ -388,7 +385,7 @@ void SVariantManager::Construct(const FArguments& InArgs, TSharedRef<FVariantMan
 				SNew(SSplitter)
 				.Orientation(Orient_Horizontal)
 				+ SSplitter::Slot()
-				.Value(ColumnSizeData.LeftColumnWidth)
+				.Value(ColumnSizeData.NameColumnWidth)
 				.OnSlotResized(ColumnSizeData.OnNameColumnResized)
 				[
 					SNew(SBox)
@@ -515,8 +512,8 @@ SVariantManager::~SVariantManager()
 			Values.VariantColumn = MainSplitter->SlotAt(0).SizeValue.Get();
 			Values.ActorColumn = MainSplitter->SlotAt(1).SizeValue.Get();
 			float PropertyCombo = MainSplitter->SlotAt(2).SizeValue.Get();
-			Values.PropertyNameColumn = PropertyCombo * OnGetLeftColumnWidth();
-			Values.PropertyValueColumn = PropertyCombo * OnGetRightColumnWidth();
+			Values.PropertyNameColumn = PropertyCombo * ColumnSizeData.NameColumnWidth.Get();
+			Values.PropertyValueColumn = PropertyCombo * ColumnSizeData.ValueColumnWidth.Get();
 
 			GConfig->SetString(TEXT("VariantManager"), TEXT("MainSplitterValues"), *Values.ToString(), GEditorPerProjectIni);
 		}
diff --git a/Engine/Plugins/Enterprise/VariantManager/Source/VariantManager/Private/SVariantManager.h b/Engine/Plugins/Enterprise/VariantManager/Source/VariantManager/Private/SVariantManager.h
--- a/Engine/Plugins/Enterprise/VariantManager/Source/VariantManager/Private/SVariantManager.h
+++ b/Engine/Plugins/Enterprise/VariantManager/Source/VariantManager/Private/SVariantManager.h
@@ -213,11 +213,6 @@ public:
 	FReply OnAddVariantSetClicked();
 	FReply OnSummonAddActorMenu();
 
-	// Callbacks for ColumnSizeData
-	float OnGetLeftColumnWidth() const { return 1.0f - RightPropertyColumnWidth; }
-	float OnGetRightColumnWidth() const { return RightPropertyColumnWidth; }
-	void OnSetColumnWidth(float InWidth) { RightPropertyColumnWidth = InWidth; }
-
 	void OnObjectTransacted(UObject* Object, const class FTransactionObjectEvent& Event);
 	void OnObjectPropertyChanged(UObject* Object, struct FPropertyChangedEvent& Event);
 	void OnPreObjectPropertyChanged(UObject* Object, const class FEditPropertyChain& PropChain);
diff --git a/Engine/Source/Editor/MaterialEditor/Private/SMaterialLayersFunctionsTree.cpp b/Engine/Source/Editor/MaterialEditor/Private/SMaterialLayersFunctionsTree.cpp
--- a/Engine/Source/Editor/MaterialEditor/Private/SMaterialLayersFunctionsTree.cpp
+++ b/Engine/Source/Editor/MaterialEditor/Private/SMaterialLayersFunctionsTree.cpp
@@ -1010,7 +1010,7 @@ void SMaterialLayersFunctionsInstanceTreeItem::Construct(const FArguments& InArg
 					.HitDetectionSplitterHandleSize(5.0f)
 					+ SSplitter::Slot()
 					.Value(Tree->ColumnSizeData.NameColumnWidth)
-					.OnSlotResized(SSplitter::FOnSlotResized::CreateLambda([](float){}))
+					.OnSlotResized(Tree->ColumnSizeData.OnNameColumnResized)
 					.Value(0.25f)
 					[
 						SNew(SHorizontalBox)
diff --git a/Engine/Source/Editor/PropertyEditor/Private/SDetailSingleItemRow.cpp b/Engine/Source/Editor/PropertyEditor/Private/SDetailSingleItemRow.cpp
--- a/Engine/Source/Editor/PropertyEditor/Private/SDetailSingleItemRow.cpp
+++ b/Engine/Source/Editor/PropertyEditor/Private/SDetailSingleItemRow.cpp
@@ -396,7 +396,7 @@ void SDetailSingleItemRow::Construct( const FArguments& InArgs, FDetailLayoutCus
 			// create inner splitter
 			OuterSplitter->AddSlot()
 				.Value(ColumnSizeData.PropertyColumnWidth)
-				.OnSlotResized(SSplitter::FOnSlotResized::CreateLambda([](float){})) // dummy handler, FDetailColumnSizeData handles this calculation
+				.OnSlotResized(ColumnSizeData.OnPropertyColumnResized)
 			[
 				SAssignNew(InnerSplitter, SSplitter)
 				.Style(FEditorStyle::Get(), "DetailsView.Splitter")
diff --git a/Engine/Source/Editor/PropertyEditor/Public/IDetailPropertyRow.h b/Engine/Source/Editor/PropertyEditor/Public/IDetailPropertyRow.h
--- a/Engine/Source/Editor/PropertyEditor/Public/IDetailPropertyRow.h
+++ b/Engine/Source/Editor/PropertyEditor/Public/IDetailPropertyRow.h
@@ -36,8 +36,9 @@ public:
 	TAttribute<float> PropertyColumnWidth;
 	TAttribute<float> RightColumnWidth;
 	TAttribute<int32> HoveredSplitterIndex;
-	SSplitter::FOnSlotResized OnNameColumnResized; // intentionally left unbound, since we don't need a handler
+	SSplitter::FOnSlotResized OnNameColumnResized; // intentionally left unbound, since we don't need a handler as the width is derived from the value width
 	SSplitter::FOnSlotResized OnValueColumnResized;
+	SSplitter::FOnSlotResized OnPropertyColumnResized; // intentionally left unbound, since we don't need a handler as the width is derived from the right width
 	SSplitter::FOnSlotResized OnRightColumnResized;
 	SSplitter::FOnHandleHovered OnSplitterHandleHovered;
 
