Dataprep: Assert after clearing previewed filters in Dataprep Recipe tab

diff --git a/Engine/Plugins/Enterprise/DataprepEditor/Source/DataprepEditor/Private/DataprepEditor.cpp b/Engine/Plugins/Enterprise/DataprepEditor/Source/DataprepEditor/Private/DataprepEditor.cpp
--- a/Engine/Plugins/Enterprise/DataprepEditor/Source/DataprepEditor/Private/DataprepEditor.cpp
+++ b/Engine/Plugins/Enterprise/DataprepEditor/Source/DataprepEditor/Private/DataprepEditor.cpp
@@ -1626,8 +1626,10 @@ void FDataprepEditor::RefreshColumnsForPreviewSystem()
 		SceneOutliner->RemoveColumn( FDataprepPreviewOutlinerColumn::ColumnID );
 		FSceneOutlinerModule& SceneOutlinerModule = FModuleManager::Get().LoadModuleChecked<FSceneOutlinerModule>("SceneOutliner");
 		FSceneOutlinerColumnInfo* ActorInfoColumPtr = SceneOutlinerModule.DefaultColumnMap.Find( FSceneOutlinerBuiltInColumnTypes::ActorInfo() );
-		check( ActorInfoColumPtr );
-		SceneOutliner->AddColumn( FSceneOutlinerBuiltInColumnTypes::ActorInfo(), *ActorInfoColumPtr );
+		if ( ActorInfoColumPtr )
+		{
+			SceneOutliner->AddColumn( FSceneOutlinerBuiltInColumnTypes::ActorInfo(), *ActorInfoColumPtr );
+		}
 
 		AssetPreviewView->RemoveColumn( FDataprepPreviewAssetColumn::ColumnID );
 	}
