Dataprep: Asset Preview background color becomes white after importing inputs data

diff --git a/Engine/Plugins/Enterprise/DataprepEditor/Source/DataprepEditor/Private/Widgets/SAssetsPreviewWidget.cpp b/Engine/Plugins/Enterprise/DataprepEditor/Source/DataprepEditor/Private/Widgets/SAssetsPreviewWidget.cpp
--- a/Engine/Plugins/Enterprise/DataprepEditor/Source/DataprepEditor/Private/Widgets/SAssetsPreviewWidget.cpp
+++ b/Engine/Plugins/Enterprise/DataprepEditor/Source/DataprepEditor/Private/Widgets/SAssetsPreviewWidget.cpp
@@ -148,7 +148,6 @@ namespace AssetPreviewWidget
 
 			SMultiColumnTableRow<IAssetTreeItemPtr>::Construct(
 				STableRow::FArguments()
-				.Style(FEditorStyle::Get(), "ContentBrowser.AssetListView.TableRow")
 				.Cursor(EMouseCursor::Default),
 				OwnerTableView);
 		}
