Dataprep: Transparent background around Add New button in Palette tab

diff --git a/Engine/Plugins/Enterprise/DataprepEditor/Source/DataprepEditor/Private/Widgets/SDataprepPalette.cpp b/Engine/Plugins/Enterprise/DataprepEditor/Source/DataprepEditor/Private/Widgets/SDataprepPalette.cpp
--- a/Engine/Plugins/Enterprise/DataprepEditor/Source/DataprepEditor/Private/Widgets/SDataprepPalette.cpp
+++ b/Engine/Plugins/Enterprise/DataprepEditor/Source/DataprepEditor/Private/Widgets/SDataprepPalette.cpp
@@ -65,68 +65,65 @@ void SDataprepPalette::Construct(const FArguments& InArgs)
 			+ SWrapBox::Slot()
 			.FillLineWhenSizeLessThan( 600 )
 			.FillEmptySpace( true )
+			.Padding( 3, 3, 3, 0 )
 			[
 				SNew( SHorizontalBox )
 
 				+ SHorizontalBox::Slot()
 				.AutoWidth()
+				.Padding( FMargin( 2 ) )
 				[
-					SNew( SBorder )
-					.Padding( FMargin( 3 ) )
-					.BorderImage( FEditorStyle::GetBrush( "ContentBrowser.TopBar.GroupBorder" ) )
-					[
-						SNew( SHorizontalBox )
+					SNew( SHorizontalBox )
 
-						// Add New
-						+ SHorizontalBox::Slot()
-						.AutoWidth()
-						.VAlign( VAlign_Center )
-						.HAlign( HAlign_Left )
+					// Add New
+					+ SHorizontalBox::Slot()
+					.AutoWidth()
+					.VAlign( VAlign_Center )
+					.HAlign( HAlign_Left )
+					[
+						SNew( SComboButton )
+						.ComboButtonStyle( FEditorStyle::Get(), "ToolbarComboButton" )
+						.ButtonStyle(FEditorStyle::Get(), "FlatButton.Success")
+						.ForegroundColor(FLinearColor::White)
+						.ContentPadding(FMargin(6, 2))
+						.OnGetMenuContent_Lambda( [this]{ return ConstructAddActionMenu(); } )
+						.HasDownArrow(false)
+						.ButtonContent()
 						[
-							SNew( SComboButton )
-							.ComboButtonStyle( FEditorStyle::Get(), "ToolbarComboButton" )
-							.ButtonStyle(FEditorStyle::Get(), "FlatButton.Success")
-							.ForegroundColor(FLinearColor::White)
-							.ContentPadding(FMargin(6, 2))
-							.OnGetMenuContent_Lambda( [this]{ return ConstructAddActionMenu(); } )
-							.HasDownArrow(false)
-							.ButtonContent()
+							SNew( SHorizontalBox )
+
+							// New Icon
+							+ SHorizontalBox::Slot()
+							.VAlign(VAlign_Center)
+							.AutoWidth()
+							[
+								SNew(STextBlock)
+								.TextStyle(FEditorStyle::Get(), "NormalText.Important")
+								.Font(FEditorStyle::Get().GetFontStyle("FontAwesome.10"))
+								.Text(FEditorFontGlyphs::File)
+							]
+
+							// New Text
+							+ SHorizontalBox::Slot()
+							.AutoWidth()
+							.VAlign(VAlign_Center)
+							.Padding(4, 0, 0, 0)
+							[
+								SNew( STextBlock )
+								.TextStyle( FEditorStyle::Get(), "NormalText.Important" )
+								.Text( LOCTEXT( "AddNewButton", "Add New" ) )
+							]
+
+							// Down Arrow
+							+ SHorizontalBox::Slot()
+							.VAlign(VAlign_Center)
+							.AutoWidth()
+							.Padding(4, 0, 0, 0)
 							[
-								SNew( SHorizontalBox )
-
-								// New Icon
-								+ SHorizontalBox::Slot()
-								.VAlign(VAlign_Center)
-								.AutoWidth()
-								[
-									SNew(STextBlock)
-									.TextStyle(FEditorStyle::Get(), "ContentBrowser.TopBar.Font")
-									.Font(FEditorStyle::Get().GetFontStyle("FontAwesome.10"))
-									.Text(FEditorFontGlyphs::File)
-								]
-
-								// New Text
-								+ SHorizontalBox::Slot()
-								.AutoWidth()
-								.VAlign(VAlign_Center)
-								.Padding(4, 0, 0, 0)
-								[
-									SNew( STextBlock )
-									.TextStyle( FEditorStyle::Get(), "ContentBrowser.TopBar.Font" )
-									.Text( LOCTEXT( "AddNewButton", "Add New" ) )
-								]
-
-								// Down Arrow
-								+ SHorizontalBox::Slot()
-								.VAlign(VAlign_Center)
-								.AutoWidth()
-								.Padding(4, 0, 0, 0)
-								[
-									SNew(STextBlock)
-									.TextStyle(FEditorStyle::Get(), "ContentBrowser.TopBar.Font")
-									.Font(FEditorStyle::Get().GetFontStyle("FontAwesome.10"))
-									.Text(FEditorFontGlyphs::Caret_Down)
-								]
+								SNew(STextBlock)
+								.TextStyle(FEditorStyle::Get(), "NormalText.Important")
+								.Font(FEditorStyle::Get().GetFontStyle("FontAwesome.10"))
+								.Text(FEditorFontGlyphs::Caret_Down)
 							]
 						]
 					]
@@ -156,7 +153,6 @@ void SDataprepPalette::Construct(const FArguments& InArgs)
 
 		+ SVerticalBox::Slot()
 		.FillHeight(1.0f)
-		.Padding(0,2,0,0)
 		[
 			SNew(SBorder)
 			.Padding(2.0f)
