Dataprep: Fix - NonUnity Compile: undefined type 'UBodySetup' in DataprepOperationsLibrary.cpp

diff --git a/Engine/Plugins/Enterprise/DataprepEditor/Source/DataprepLibraries/Private/DataprepOperationsLibrary.cpp b/Engine/Plugins/Enterprise/DataprepEditor/Source/DataprepLibraries/Private/DataprepOperationsLibrary.cpp
--- a/Engine/Plugins/Enterprise/DataprepEditor/Source/DataprepLibraries/Private/DataprepOperationsLibrary.cpp
+++ b/Engine/Plugins/Enterprise/DataprepEditor/Source/DataprepLibraries/Private/DataprepOperationsLibrary.cpp
@@ -30,6 +30,7 @@
 #include "Math/Vector2D.h"
 #include "Misc/FileHelper.h"
 #include "ObjectTools.h"
+#include "PhysicsEngine/BodySetup.h"
 #include "StaticMeshAttributes.h"
 #include "StaticMeshOperations.h"
 #include "TessellationRendering.h"
