Fixed ensure failure after import in Dataprep
Removed API redundancy between DataprepCore and DataprepEditor

diff --git a/Engine/Plugins/Enterprise/DataprepEditor/Source/DataprepCore/Private/Shared/DataprepCorePrivateUtils.h b/Engine/Plugins/Enterprise/DataprepEditor/Source/DataprepCore/Private/Shared/DataprepCorePrivateUtils.h
--- a/Engine/Plugins/Enterprise/DataprepEditor/Source/DataprepCore/Private/Shared/DataprepCorePrivateUtils.h
+++ b/Engine/Plugins/Enterprise/DataprepEditor/Source/DataprepCore/Private/Shared/DataprepCorePrivateUtils.h
@@ -54,7 +54,7 @@ namespace DataprepCorePrivateUtils
 	void DeleteRegisteredAsset(UObject* Asset);
 
 	/** Returns directory where to store temporary files when running Dataprep asset */
-	const FString& GetRootTemporaryDir();
+	DATAPREPCORE_API const FString& GetRootTemporaryDir();
 
 	/** Returns content folder where to create temporary assets when running Dataprep asset */
 	DATAPREPCORE_API const FString& GetRootPackagePath();
diff --git a/Engine/Plugins/Enterprise/DataprepEditor/Source/DataprepEditor/Private/DataprepEditor.cpp b/Engine/Plugins/Enterprise/DataprepEditor/Source/DataprepEditor/Private/DataprepEditor.cpp
--- a/Engine/Plugins/Enterprise/DataprepEditor/Source/DataprepEditor/Private/DataprepEditor.cpp
+++ b/Engine/Plugins/Enterprise/DataprepEditor/Source/DataprepEditor/Private/DataprepEditor.cpp
@@ -273,9 +273,8 @@ FDataprepEditor::FDataprepEditor()
 	, bIgnoreCloseRequest(false)
 	, PreviewSystem( MakeShared<FDataprepPreviewSystem>() )
 {
-	UPackage* DataprepPackage = NewObject<UPackage>(nullptr, *GetRootPackagePath(), RF_Transient);
-	FName UniqueWorldName = MakeUniqueObjectName(DataprepPackage, UWorld::StaticClass(), FName( *(LOCTEXT("PreviewWorld", "Preview").ToString()) ));
-	PreviewWorld = NewObject<UWorld>(DataprepPackage, UniqueWorldName);
+	FName UniqueWorldName = MakeUniqueObjectName(GetTransientPackage(), UWorld::StaticClass(), FName( *(LOCTEXT("PreviewWorld", "Preview").ToString()) ));
+	PreviewWorld = NewObject<UWorld>(GetTransientPackage(), UniqueWorldName);
 	PreviewWorld->WorldType = EWorldType::EditorPreview;
 
 	FWorldContext& WorldContext = GEngine->CreateNewWorldContext(PreviewWorld->WorldType);
@@ -346,13 +345,13 @@ FDataprepEditor::~FDataprepEditor()
 			return bDirectoryIsEmpty;
 		};
 
-		FString RootTempDir = FPaths::Combine( GetRootTemporaryDir(), FString::FromInt( FPlatformProcess::GetCurrentProcessId() ) );
+		FString RootTempDir = FPaths::Combine( DataprepCorePrivateUtils::GetRootTemporaryDir(), FString::FromInt( FPlatformProcess::GetCurrentProcessId() ) );
 		if(IsDirectoryEmpty( *RootTempDir ))
 		{
 			DeleteDirectory( RootTempDir );
 		}
 
-		const FString PackagePathToDelete = FPaths::Combine( GetRootPackagePath(), FString::FromInt( FPlatformProcess::GetCurrentProcessId() ) );
+		const FString PackagePathToDelete = FPaths::Combine( DataprepCorePrivateUtils::GetRootPackagePath(), FString::FromInt( FPlatformProcess::GetCurrentProcessId() ) );
 		FString PackagePathToDeleteOnDisk;
 		if (FPackageName::TryConvertLongPackageNameToFilename(PackagePathToDelete, PackagePathToDeleteOnDisk))
 		{
@@ -449,7 +448,7 @@ void FDataprepEditor::CleanUpTemporaryDirectories()
 	const int32 CurrentProcessID = FPlatformProcess::GetCurrentProcessId();
 
 	TSet<FString> TempDirectories;
-	IFileManager::Get().IterateDirectory( *GetRootTemporaryDir(), [&](const TCHAR* FilenameOrDirectory, bool bIsDirectory) -> bool
+	IFileManager::Get().IterateDirectory( *DataprepCorePrivateUtils::GetRootTemporaryDir(), [&](const TCHAR* FilenameOrDirectory, bool bIsDirectory) -> bool
 	{
 		if (bIsDirectory)
 		{
@@ -474,7 +473,7 @@ void FDataprepEditor::CleanUpTemporaryDirectories()
 
 					if(bDeleteDirectories)
 					{
-						FString PackagePathToDelete = FPaths::Combine( GetRootPackagePath(), DirectoryName );
+						FString PackagePathToDelete = FPaths::Combine( DataprepCorePrivateUtils::GetRootPackagePath(), DirectoryName );
 						FString PackagePathToDeleteOnDisk;
 						if (FPackageName::TryConvertLongPackageNameToFilename(PackagePathToDelete, PackagePathToDeleteOnDisk))
 						{
@@ -496,17 +495,6 @@ void FDataprepEditor::CleanUpTemporaryDirectories()
 	}
 }
 
-const FString& FDataprepEditor::GetRootTemporaryDir()
-{
-	static FString RootTemporaryDir = FPaths::Combine(FPaths::ProjectIntermediateDir(), TEXT("DataprepTemp") );
-	return RootTemporaryDir;
-}
-
-const FString& FDataprepEditor::GetRootPackagePath()
-{
-	return DataprepCorePrivateUtils::GetRootPackagePath();
-}
-
 void FDataprepEditor::InitDataprepEditor(const EToolkitMode::Type Mode, const TSharedPtr<IToolkitHost>& InitToolkitHost, UDataprepAssetInterface* InDataprepAssetInterface)
 {
 	DataprepAssetInterfacePtr = TWeakObjectPtr<UDataprepAssetInterface>(InDataprepAssetInterface);
@@ -544,7 +532,7 @@ void FDataprepEditor::InitDataprepEditor(const EToolkitMode::Type Mode, const TS
 
 	// Create temporary directory to store transient data
 	CleanUpTemporaryDirectories();
-	TempDir = FPaths::Combine( GetRootTemporaryDir(), FString::FromInt( FPlatformProcess::GetCurrentProcessId() ), SessionID);
+	TempDir = FPaths::Combine( DataprepCorePrivateUtils::GetRootTemporaryDir(), FString::FromInt( FPlatformProcess::GetCurrentProcessId() ), SessionID);
 	IFileManager::Get().MakeDirectory(*TempDir);
 
 	GEditor->RegisterForUndo(this);
@@ -1591,7 +1579,7 @@ bool FDataprepEditor::CanCommitWorld()
 
 FString FDataprepEditor::GetTransientContentFolder()
 {
-	return FPaths::Combine( GetRootPackagePath(), FString::FromInt( FPlatformProcess::GetCurrentProcessId() ), SessionID );
+	return FPaths::Combine( DataprepCorePrivateUtils::GetRootPackagePath(), FString::FromInt( FPlatformProcess::GetCurrentProcessId() ), SessionID );
 }
 
 bool FDataprepEditor::OnCanExecuteNextStep(UDataprepActionAsset* ActionAsset)
diff --git a/Engine/Plugins/Enterprise/DataprepEditor/Source/DataprepEditor/Private/DataprepEditor.h b/Engine/Plugins/Enterprise/DataprepEditor/Source/DataprepEditor/Private/DataprepEditor.h
--- a/Engine/Plugins/Enterprise/DataprepEditor/Source/DataprepEditor/Private/DataprepEditor.h
+++ b/Engine/Plugins/Enterprise/DataprepEditor/Source/DataprepEditor/Private/DataprepEditor.h
@@ -105,12 +105,6 @@ public:
 	/** Gets or sets the flag for context sensitivity in the graph action menu */
 	bool& GetIsContextSensitive() { return bIsActionMenuContextSensitive; }
 
-	/** Returns root package which all transient packages are created under */
-	static const FString& GetRootPackagePath();
-
-	/** Returns root directory which all transient directories and data are created under */
-	static const FString& GetRootTemporaryDir();
-
 	/** Return the selected object from the scene outliner / world preview */
 	const TSet<TWeakObjectPtr<UObject>>& GetWorldItemsSelection() const { return WorldItemsSelection; };
 
