Fixed crash on second execution of a Dataprep pipeline in the editor
  - Root cause: Changes made to UStaticMesh
  - Solution: Revisited serialization logic when creating and reading a snapshot of a static mesh. This is a temporary fix which will need more work, see UEENT-4385
Fixed crash when undoing any operation in Dataprep pipeline editor
  - Root cause: check on a null pointer
  - Solution: replace check with an if.

diff --git a/Engine/Plugins/Enterprise/DataprepEditor/Source/DataprepEditor/Private/DataprepSnapshot.cpp b/Engine/Plugins/Enterprise/DataprepEditor/Source/DataprepEditor/Private/DataprepSnapshot.cpp
--- a/Engine/Plugins/Enterprise/DataprepEditor/Source/DataprepEditor/Private/DataprepSnapshot.cpp
+++ b/Engine/Plugins/Enterprise/DataprepEditor/Source/DataprepEditor/Private/DataprepSnapshot.cpp
@@ -310,8 +310,10 @@ namespace DataprepSnapshotUtil
 		MemAr << SubObjectsCount;
 
 		// Serialize class and path of each sub-object
-		for(UObject* SubObject : SubObjectsArray)
+		for(int32 Index = SubObjectsArray.Num() - 1; Index >= 0; --Index)
 		{
+			const UObject* SubObject = SubObjectsArray[Index];
+
 			UClass* SubObjectClass = SubObject->GetClass();
 
 			FString ClassName = SubObjectClass->GetName();
@@ -337,16 +339,26 @@ namespace DataprepSnapshotUtil
 			MemAr << SubObjectName;
 		}
 
+		bool bIsStaticMesh = Cast<UStaticMesh>(Object) != nullptr;
+
+		// Overwrite persistent flag if Object is a static mesh.
+		// This will skip the persistence of the MeshDescription on static mesh and geometrical sub-objects
+		Ar.SetIsPersistent(bIsStaticMesh);
+
 		// Serialize sub-objects' content
 		for(UObject* SubObject : SubObjectsArray)
 		{
-			SubObject->Serialize(Ar);
+			if (!bIsStaticMesh || !SubObject->HasAnyFlags(RF_DefaultSubObject))
+			{
+				SubObject->Serialize(Ar);
+			}
 		}
 
+		// Overwrite transacting flag if Object is a static mesh.
+		// This will skip the persistence of the MeshDescription
+		Ar.SetIsTransacting(!bIsStaticMesh);
+
 		// Serialize object
-		// Overwrite transacting and persistent flags if Object is a static mesh. This will skip the persistence of the MeshDescription
-		Ar.SetIsTransacting(Cast<UStaticMesh>(Object) == nullptr);
-		Ar.SetIsPersistent(Cast<UStaticMesh>(Object) != nullptr);
 		Object->Serialize(Ar);
 
 		if(UTexture* Texture = Cast<UTexture>(Object))
@@ -358,17 +370,23 @@ namespace DataprepSnapshotUtil
 
 	void ReadSnapshotData(UObject* Object, const FSnapshotBuffer& InSerializedData, TMap<FString, UClass*>& InClassesMap, TArray<UObject*>& ObjectsToDelete)
 	{
+		bool bIsStaticMesh = Cast<UStaticMesh>(Object) != nullptr;
+
 		// Remove all objects created by default that InObject is dependent on
 		// This method must obviously be called just after the InObject is created
-		auto RemoveDefaultDependencies = [&ObjectsToDelete](UObject* InObject)
+		auto RemoveDefaultDependencies = [&ObjectsToDelete, bIsStaticMesh](UObject* InObject)
 		{
 			TArray< UObject* > ObjectsWithOuter;
 			GetObjectsWithOuter( InObject, ObjectsWithOuter, /*bIncludeNestedObjects = */ true );
 
 			for(UObject* ObjectWithOuter : ObjectsWithOuter)
 			{
-				FDataprepCoreUtils::MoveToTransientPackage( ObjectWithOuter );
-				ObjectsToDelete.Add( ObjectWithOuter );
+				// Do not delete default sub-objects
+				if (!bIsStaticMesh || !ObjectWithOuter->HasAnyFlags(RF_DefaultSubObject))
+				{
+					FDataprepCoreUtils::MoveToTransientPackage( ObjectWithOuter );
+					ObjectsToDelete.Add( ObjectWithOuter );
+				}
 			}
 		};
 
@@ -383,9 +401,12 @@ namespace DataprepSnapshotUtil
 
 		// Create empty sub-objects based on class and patch
 		TArray< UObject* > SubObjectsArray;
-		SubObjectsArray.Reserve(SubObjectsCount);
+		SubObjectsArray.SetNumZeroed(SubObjectsCount);
+
+		// Create root name to avoid name collision
+		FString RootName = FGuid::NewGuid().ToString();
 
-		for(int32 Index = 0; Index < SubObjectsCount; ++Index)
+		for(int32 Index = SubObjectsCount - 1; Index >= 0; --Index)
 		{
 			FString ClassName;
 			MemAr << ClassName;
@@ -396,8 +417,17 @@ namespace DataprepSnapshotUtil
 			int32 ObjectFlags;
 			MemAr << ObjectFlags;
 
-			UObject* SubObject = NewObject<UObject>( Object, SubObjectClass, NAME_None, EObjectFlags(ObjectFlags) );
-			SubObjectsArray.Add( SubObject );
+			// Temporary - Do not create default sub-object for static mesh
+			// #ueent_dataprep: TODO: Implement a more robust serialization process for transient UProperties
+			if (bIsStaticMesh && (ObjectFlags & RF_DefaultSubObject))
+			{
+				continue;
+			}
+
+			FString SubObjectName = RootName + FString::FromInt(Index);
+
+			UObject* SubObject = NewObject<UObject>( Object, SubObjectClass, *SubObjectName, EObjectFlags(ObjectFlags) );
+			SubObjectsArray[Index] = SubObject;
 
 			RemoveDefaultDependencies( SubObject );
 		}
@@ -417,23 +447,38 @@ namespace DataprepSnapshotUtil
 			UObject* NewOuter = SoftPath.ResolveObject();
 			ensure( NewOuter );
 
-			UObject* SubObject = SubObjectsArray[Index];
-			if( NewOuter != SubObject->GetOuter() )
+			if (UObject* SubObject = SubObjectsArray[Index])
 			{
-				FDataprepCoreUtils::RenameObject( SubObject, *SubObjectName, NewOuter );
+				if( NewOuter != SubObject->GetOuter() )
+				{
+					FDataprepCoreUtils::RenameObject( SubObject, nullptr, NewOuter );
+				}
+
+				if( SubObjectName != SubObject->GetName() && SubObject->Rename( *SubObjectName, nullptr, REN_Test ))
+				{
+					FDataprepCoreUtils::RenameObject( SubObject, *SubObjectName );
+				}
 			}
 		}
 
+		// Overwrite persistent flag if Object is a static mesh.
+		// This will skip the persistence of the MeshDescription on static mesh and geometrical sub-objects
+		Ar.SetIsPersistent(bIsStaticMesh);
+
 		// Deserialize sub-objects
 		for(UObject* SubObject : SubObjectsArray)
 		{
-			SubObject->Serialize(Ar);
+			if (SubObject)
+			{
+				SubObject->Serialize(Ar);
+			}
 		}
 
+		// Overwrite transacting flag if Object is a static mesh.
+		// This will skip the persistence of the MeshDescription
+		Ar.SetIsTransacting(!bIsStaticMesh);
+
 		// Deserialize object
-		// Overwrite transacting and persistent flags if Object is a static mesh. This will skip the persistence of the MeshDescription
-		Ar.SetIsTransacting(Cast<UStaticMesh>(Object) == nullptr);
-		Ar.SetIsPersistent(Cast<UStaticMesh>(Object) != nullptr);
 		Object->Serialize(Ar);
 
 		if(UTexture* Texture = Cast<UTexture>(Object))
