Fixed unreported crash when importing a file in the Dataprep editor

diff --git a/Engine/Plugins/Enterprise/DataprepEditor/Source/DataprepEditor/Private/DataprepSnapshot.cpp b/Engine/Plugins/Enterprise/DataprepEditor/Source/DataprepEditor/Private/DataprepSnapshot.cpp
--- a/Engine/Plugins/Enterprise/DataprepEditor/Source/DataprepEditor/Private/DataprepSnapshot.cpp
+++ b/Engine/Plugins/Enterprise/DataprepEditor/Source/DataprepEditor/Private/DataprepSnapshot.cpp
@@ -13,6 +13,7 @@
 #include "Async/ParallelFor.h"
 #include "Editor/UnrealEdEngine.h"
 #include "EngineUtils.h"
+#include "Engine/StaticMeshSourceData.h"
 #include "Engine/Texture.h"
 #include "Exporters/Exporter.h"
 #include "Factories/LevelFactory.h"
@@ -213,6 +214,36 @@ namespace DataprepSnapshotUtil
 		return FPaths::ConvertRelativePathToFull( FPaths::Combine( RootPath, PackageFileName ) + SnapshotExtension );
 	}
 
+	// #ueent_dataprep: Revisit serialization of snapshot in 5.0
+	void SerializeObject(FSnapshotCustomArchive& Ar, UObject* Object, TArray< UObject* >& SubObjects)
+	{
+		bool bIsStaticMesh = Cast<UStaticMesh>(Object) != nullptr;
+
+		// Overwrite persistent flag if Object is a static mesh.
+		// This will skip the persistence of the MeshDescription on static mesh and geometrical sub-objects
+
+		// Serialize SubObjects' content
+		for(UObject* SubObject : SubObjects)
+		{
+			if (SubObject && (!bIsStaticMesh || !SubObject->HasAnyFlags(RF_DefaultSubObject)))
+			{
+				// Do not set the archive as persistent if the subobject is mesh description's bulk data.
+				// The combination of transient object and persistence is not accepted by such object during serialization
+				Ar.SetIsPersistent(bIsStaticMesh && SubObject->StaticClass()->IsChildOf(UMeshDescriptionBaseBulkData::StaticClass()));
+				SubObject->Serialize(Ar);
+			}
+		}
+
+		// Overwrite transacting flag if Object is a static mesh.
+		// This will skip the persistence of the MeshDescription
+		// Set archive as persistent if dealing with a static mesh. Otherwise nothingis serialized
+		Ar.SetIsPersistent(bIsStaticMesh);
+		Ar.SetIsTransacting(!bIsStaticMesh);
+
+		// Serialize object
+		Object->Serialize(Ar);
+	}
+
 	void WriteSnapshotData(UObject* Object, FSnapshotBuffer& OutSerializedData)
 	{
 		// Helper struct to identify dependency of a UObject on other UObject(s) except given one (its outer)
@@ -339,33 +370,7 @@ namespace DataprepSnapshotUtil
 			MemAr << SubObjectName;
 		}
 
-		bool bIsStaticMesh = Cast<UStaticMesh>(Object) != nullptr;
-
-		// Overwrite persistent flag if Object is a static mesh.
-		// This will skip the persistence of the MeshDescription on static mesh and geometrical sub-objects
-		Ar.SetIsPersistent(bIsStaticMesh);
-
-		// Serialize sub-objects' content
-		for(UObject* SubObject : SubObjectsArray)
-		{
-			if (!bIsStaticMesh || !SubObject->HasAnyFlags(RF_DefaultSubObject))
-			{
-				SubObject->Serialize(Ar);
-			}
-		}
-
-		// Overwrite transacting flag if Object is a static mesh.
-		// This will skip the persistence of the MeshDescription
-		Ar.SetIsTransacting(!bIsStaticMesh);
-
-		// Serialize object
-		Object->Serialize(Ar);
-
-		if(UTexture* Texture = Cast<UTexture>(Object))
-		{
-			bool bRebuildResource = !!Texture->Resource;
-			Ar << bRebuildResource;
-		}
+		SerializeObject(Ar, Object, SubObjectsArray);
 	}
 
 	void ReadSnapshotData(UObject* Object, const FSnapshotBuffer& InSerializedData, TMap<FString, UClass*>& InClassesMap, TArray<UObject*>& ObjectsToDelete)
@@ -461,35 +466,11 @@ namespace DataprepSnapshotUtil
 			}
 		}
 
-		// Overwrite persistent flag if Object is a static mesh.
-		// This will skip the persistence of the MeshDescription on static mesh and geometrical sub-objects
-		Ar.SetIsPersistent(bIsStaticMesh);
-
-		// Deserialize sub-objects
-		for(UObject* SubObject : SubObjectsArray)
-		{
-			if (SubObject)
-			{
-				SubObject->Serialize(Ar);
-			}
-		}
-
-		// Overwrite transacting flag if Object is a static mesh.
-		// This will skip the persistence of the MeshDescription
-		Ar.SetIsTransacting(!bIsStaticMesh);
-
-		// Deserialize object
-		Object->Serialize(Ar);
+		SerializeObject(Ar, Object, SubObjectsArray);
 
 		if(UTexture* Texture = Cast<UTexture>(Object))
 		{
-			bool bRebuildResource = false;
-			Ar << bRebuildResource;
-
-			if(bRebuildResource)
-			{
-				Texture->UpdateResource();
-			}
+			Texture->UpdateResource();
 		}
 	}
 }
