Slate: Convert SConstraint FSlot's to the FSlotArgument syntax.

diff --git a/Engine/Plugins/Enterprise/DataprepEditor/Source/DataprepEditor/Private/Widgets/DataprepGraph/SDataprepGraphTrackNode.cpp b/Engine/Plugins/Enterprise/DataprepEditor/Source/DataprepEditor/Private/Widgets/DataprepGraph/SDataprepGraphTrackNode.cpp
--- a/Engine/Plugins/Enterprise/DataprepEditor/Source/DataprepEditor/Private/Widgets/DataprepGraph/SDataprepGraphTrackNode.cpp
+++ b/Engine/Plugins/Enterprise/DataprepEditor/Source/DataprepEditor/Private/Widgets/DataprepGraph/SDataprepGraphTrackNode.cpp
@@ -188,7 +188,7 @@ public:
 
 	FVector2D GetSoftwareCursorPosition() const
 	{
-		return FVector2D( DragSlot->OffsetAttr.Get(FMargin()).Left + TrackCursorOffset.X, TrackCursorOffset.Y);
+		return FVector2D( DragSlot->GetOffset().Left + TrackCursorOffset.X, TrackCursorOffset.Y);
 	}
 
 	FVector2D GetAbsoluteSoftwareCursorPosition() const
@@ -209,8 +209,8 @@ public:
 	{
 		DragIndicatorIndex = INDEX_NONE;
 
-		FMargin SlotOffset = ActionSlots[SlotCount]->OffsetAttr.Get(FMargin());
-		TrackSlots[SlotCount]->OffsetAttr.Set( FMargin(SlotOffset.Left + SDataprepGraphActionNode::DefaultWidth * 0.5f - 16.f, LineTopPadding + TrackSlotTopOffset, 32, 32));
+		FMargin SlotOffset = ActionSlots[SlotCount]->GetOffset();
+		TrackSlots[SlotCount]->SetOffset( FMargin(SlotOffset.Left + SDataprepGraphActionNode::DefaultWidth * 0.5f - 16.f, LineTopPadding + TrackSlotTopOffset, 32, 32));
 		TrackSlots[SlotCount]->AttachWidget(SNullWidget::NullWidget);
 
 		UpdateLine(SlotCount-1);
@@ -218,7 +218,7 @@ public:
 
 	FVector2D UpdateActionSlot(FMargin& SlotOffset, int32 SlotIndex, const TSharedRef<SWidget>& ActionNode, bool bIsEmpty = false)
 	{
-		DropSlots[SlotIndex]->OffsetAttr.Set(SlotOffset);
+		DropSlots[SlotIndex]->SetOffset(SlotOffset);
 		SlotOffset.Left += InterSlotSpacing;
 
 		const float NodeTopPadding = LineTopPadding + NodeTopOffset;
@@ -228,10 +228,10 @@ public:
 
 		FVector2D Size = GetNodeSize(ActionNode);
 
-		ActionSlot->OffsetAttr.Set(FMargin(SlotOffset.Left, NodeTopPadding, Size.X, Size.Y));
+		ActionSlot->SetOffset(FMargin(SlotOffset.Left, NodeTopPadding, Size.X, Size.Y));
 		ActionSlot->AttachWidget(bIsEmpty ? SNew(SColorBlock).Color(FLinearColor::Transparent).Size( Size ) : ActionNode);
 
-		TrackSlots[SlotIndex]->OffsetAttr.Set( FMargin(SlotOffset.Left + Size.X * 0.5f - 16.f, TrackSlotTopPadding, 32, 32));
+		TrackSlots[SlotIndex]->SetOffset( FMargin(SlotOffset.Left + Size.X * 0.5f - 16.f, TrackSlotTopPadding, 32, 32));
 		TrackSlots[SlotIndex]->AttachWidget(SlotIndex == IndexHovered ? TrackSlotDrop.ToSharedRef() : TrackSlotRegular.ToSharedRef());
 
 		SlotOffset.Left += Size.X;
@@ -693,7 +693,7 @@ FSlateRect SDataprepGraphTrackNode::Update()
 				CanvasOffset.Bottom = TargetSize.Y;
 			}
 
-			TrackWidgetPtr->CanvasSlot->OffsetAttr.Set(CanvasOffset);
+			TrackWidgetPtr->CanvasSlot->SetOffset(CanvasOffset);
 		}
 
 		WorkingArea = FSlateRect(FVector2D::ZeroVector, WorkingSize + TrackPosition);
@@ -1152,24 +1152,26 @@ void SDataprepGraphTrackWidget::Construct(const FArguments& InArgs, TSharedPtr<S
 
 	const FLinearColor TrackColor = FDataprepEditorStyle::GetColor( "DataprepAction.BackgroundColor" );
 
-	CanvasSlot = &AddSlot()
+	AddSlot()
 	.Anchors( FAnchors( 0.f, 0.f, 0.f, 0.f ) )
 	.Offset( FMargin(-25.f, -25.f, 0.f, 0.f) )
 	.AutoSize(false)
 	.Alignment(FVector2D::ZeroVector)
 	.ZOrder(0)
+	.Expose(CanvasSlot)
 	[
 		SNew(SColorBlock)
 		.Color( FDataprepEditorStyle::GetColor("Dataprep.Background.Black") )
 		.Size( TAttribute<FVector2D>::Create(TAttribute<FVector2D>::FGetter::CreateSP( this, &SDataprepGraphTrackWidget::GetCanvasArea ) ) )
 	];
 
-	LineSlot = &AddSlot()
+	AddSlot()
 	.Anchors( FAnchors( 0.f, 0.f, 0.f, 0.f ) )
 	.Offset( FMargin(InterSlotSpacing +  3.f, LineTopPadding - 2.f, 0.f, 0.f) )
 	.AutoSize(true)
 	.Alignment(FVector2D::ZeroVector)
 	.ZOrder(0)
+	.Expose(LineSlot)
 	[
 		SNew(SColorBlock)
 		.Color( TrackColor )
@@ -1218,15 +1220,17 @@ void SDataprepGraphTrackWidget::Construct(const FArguments& InArgs, TSharedPtr<S
 		{
 			// Add drop slot ahead of action node
 			{
-				SConstraintCanvas::FSlot& Slot = AddSlot()
+				SConstraintCanvas::FSlot* Slot = nullptr;
+				AddSlot()
 				.Anchors( FAnchors( 0., 0.f, 0.f, 0.f ) )
 				.Offset( FMargin(LeftOffset, NodeTopPadding + 5.f, InterSlotSpacing, TrackDesiredHeight ) )
 				.Alignment(FVector2D::ZeroVector)
 				.ZOrder(2)
+				.Expose(Slot)
 				[
 					DropFiller.ToSharedRef()
 				];
-				DropSlots.Add( &Slot );
+				DropSlots.Add( Slot );
 			}
 			LeftOffset += InterSlotSpacing;
 
@@ -1236,30 +1240,34 @@ void SDataprepGraphTrackWidget::Construct(const FArguments& InArgs, TSharedPtr<S
 			{
 				TSharedRef<SWidget>& ActionNode = ActionNodes.Add_GetRef(InActionNodes[Index].IsValid() ? InActionNodes[Index]->AsShared() : NullNode->AsShared());
 
-				SConstraintCanvas::FSlot& Slot = AddSlot()
+				SConstraintCanvas::FSlot* Slot = nullptr;
+				AddSlot()
 				.Anchors( FAnchors( 0., 0.f, 0.f, 0.f ) )
 				.Offset( FMargin(LeftOffset, NodeTopPadding, Size.X, Size.Y ) )
 				.Alignment(FVector2D::ZeroVector)
 				.AutoSize(true)
 				.ZOrder(2)
+				.Expose(Slot)
 				[
 					ActionNode
 				];
-				ActionSlots.Add( &Slot );
+				ActionSlots.Add( Slot );
 			}
 
 			// Add track slot which is at the middle of the corresponding action
 			{
-				SConstraintCanvas::FSlot& Slot = AddSlot()
+				SConstraintCanvas::FSlot* Slot = nullptr;
+				AddSlot()
 				.Anchors( FAnchors( 0., 0.f, 0.f, 0.f ) )
 				.Offset( FMargin(LeftOffset + SDataprepGraphActionNode::DefaultWidth * 0.5f + 16.f, TrackSlotTopPadding, 32, 32 ) )
 				.Alignment(FVector2D::ZeroVector)
 				.AutoSize(true)
 				.ZOrder(2)
+				.Expose(Slot)
 				[
 					TrackSlotRegular.ToSharedRef()
 				];
-				TrackSlots.Add( &Slot );
+				TrackSlots.Add( Slot );
 			}
 			LeftOffset += Size.X;
 
@@ -1273,68 +1281,77 @@ void SDataprepGraphTrackWidget::Construct(const FArguments& InArgs, TSharedPtr<S
 
 	// Add drop slot at end of track
 	{
-		SConstraintCanvas::FSlot& Slot = AddSlot()
+		SConstraintCanvas::FSlot* Slot = nullptr;
+		AddSlot()
 		.Anchors( FAnchors( 0., 0.f, 0.f, 0.f ) )
 		.Offset( FMargin(LeftOffset, NodeTopPadding + 5.f, InterSlotSpacing, TrackDesiredHeight ) )
 		.Alignment(FVector2D::ZeroVector)
 		.ZOrder(2)
+		.Expose(Slot)
 		[
 			SNullWidget::NullWidget
 		];
-		DropSlots.Add( &Slot );
+		DropSlots.Add( Slot );
 		LeftOffset += InterSlotSpacing;
 	}
 
 	// Add dummy action slot at end of track
 	{
-		SConstraintCanvas::FSlot& Slot = AddSlot()
+		SConstraintCanvas::FSlot* Slot = nullptr;
+		AddSlot()
 		.Anchors( FAnchors( 0., 0.f, 0.f, 0.f ) )
 		.Offset( FMargin(LeftOffset, TrackDesiredHeight, SDataprepGraphActionNode::DefaultWidth, SDataprepGraphActionNode::DefaultWidth ) )
 		.Alignment(FVector2D::ZeroVector)
 		.AutoSize(true)
 		.ZOrder(2)
+		.Expose(Slot)
 		[
 			SNullWidget::NullWidget
 		];
-		ActionSlots.Add( &Slot );
+		ActionSlots.Add( Slot );
 	}
 
 	// Add track slot which is at the middle of the corresponding action
 	{
-		SConstraintCanvas::FSlot& Slot = AddSlot()
+		SConstraintCanvas::FSlot* Slot = nullptr;
+		AddSlot()
 		.Anchors( FAnchors( 0., 0.f, 0.f, 0.f ) )
 		.Offset( FMargin(LeftOffset + SDataprepGraphActionNode::DefaultWidth * 0.5f + 16.f, TrackSlotTopPadding, 32, 32 ) )
 		.Alignment(FVector2D::ZeroVector)
 		.AutoSize(true)
 		.ZOrder(2)
+		.Expose(Slot)
 		[
 			SNullWidget::NullWidget
 		];
-		TrackSlots.Add( &Slot );
+		TrackSlots.Add( Slot );
 	}
 	LeftOffset += SDataprepGraphActionNode::DefaultWidth;
 
 	// Add drop slot at end of track
 	{
-		SConstraintCanvas::FSlot& Slot = AddSlot()
+		SConstraintCanvas::FSlot* Slot = nullptr;
+		AddSlot()
 		.Anchors( FAnchors( 0., 0.f, 0.f, 0.f ) )
 		.Offset( FMargin(LeftOffset, NodeTopPadding + 5.f, InterSlotSpacing, TrackDesiredHeight ) )
 		.Alignment(FVector2D::ZeroVector)
 		.ZOrder(2)
+		.Expose(Slot)
 		[
 			SNullWidget::NullWidget
 		];
-		DropSlots.Add( &Slot );
+		DropSlots.Add( Slot );
 		LeftOffset += InterSlotSpacing;
 	}
 
 	// Add slot to host dragged node???
-	DragSlot = &AddSlot()
+	AddSlot()
 	.Anchors( FAnchors( 0., 0.f, 0.f, 0.f ) )
 	.Offset( FMargin(LeftOffset, 0.f, SDataprepGraphActionNode::DefaultWidth, SDataprepGraphActionNode::DefaultWidth ) )
 	.Alignment(FVector2D::ZeroVector)
 	.AutoSize(true)
 	.ZOrder(3)
+	.Expose(DragSlot)
 	[
 		SNullWidget::NullWidget
 	];
@@ -1416,9 +1433,9 @@ int32 SDataprepGraphTrackWidget::GetHoveredActionNode(float LeftCorner, float Wi
 	// Compute origin abscissa of last action node on track
 	for(int32 Index = 0; Index < SlotCount; ++Index)
 	{
-		const FMargin SlotOffset = DropSlots[Index]->OffsetAttr.Get(FMargin());
+		const FMargin SlotOffset = DropSlots[Index]->GetOffset();
 		const float RelativeAbscissa = CenterAbscissa - (SlotOffset.Left + SlotOffset.Right);
-		const float ActionSlotWidth = ActionSlots[Index]->OffsetAttr.Get(FMargin()).Right;
+		const float ActionSlotWidth = ActionSlots[Index]->GetOffset().Right;
 		if(RelativeAbscissa > -SlotOffset.Right && RelativeAbscissa <= ActionSlotWidth)
 		{
 			return Index;
@@ -1446,7 +1463,7 @@ void SDataprepGraphTrackWidget::RefreshLayout()
 
 	LineSize = FVector2D::ZeroVector;
 	WorkingSize = FVector2D(SDataprepGraphActionNode::DefaultWidth, SDataprepGraphActionNode::DefaultHeight);
-	FMargin DropSlotOffset = DropSlots[0]->OffsetAttr.Get(FMargin());
+	FMargin DropSlotOffset = DropSlots[0]->GetOffset();
 
 	float MaxWidth = SDataprepGraphActionNode::DefaultWidth;
 
@@ -1470,14 +1487,14 @@ void SDataprepGraphTrackWidget::RefreshLayout()
 		}
 	}
 
-	DropSlots[SlotCount]->OffsetAttr.Set(DropSlotOffset);
+	DropSlots[SlotCount]->SetOffset(DropSlotOffset);
 	DropSlots[SlotCount]->AttachWidget(SNullWidget::NullWidget);
 	DropSlotOffset.Left += InterSlotSpacing;
 
-	ActionSlots[SlotCount]->OffsetAttr.Set( FMargin(DropSlotOffset.Left, NodeTopPadding, 0., 0.));
+	ActionSlots[SlotCount]->SetOffset( FMargin(DropSlotOffset.Left, NodeTopPadding, 0., 0.));
 	ActionSlots[SlotCount]->AttachWidget(SNullWidget::NullWidget);
 
-	TrackSlots[SlotCount]->OffsetAttr.Set( FMargin(DropSlotOffset.Left + SDataprepGraphActionNode::DefaultWidth * 0.5f - 16.f, LineTopPadding + TrackSlotTopOffset, 32, 32));
+	TrackSlots[SlotCount]->SetOffset( FMargin(DropSlotOffset.Left + SDataprepGraphActionNode::DefaultWidth * 0.5f - 16.f, LineTopPadding + TrackSlotTopOffset, 32, 32));
 	TrackSlots[SlotCount]->AttachWidget(SNullWidget::NullWidget);
 
 	WorkingSize.X = DropSlotOffset.Left + SDataprepGraphActionNode::DefaultWidth + InterSlotSpacing;
@@ -1490,7 +1507,7 @@ void SDataprepGraphTrackWidget::RefreshLayout()
 
 void SDataprepGraphTrackWidget::UpdateLayoutOnDrag()
 {
-	FMargin SlotOffset = DropSlots[0]->OffsetAttr.Get(FMargin());
+	FMargin SlotOffset = DropSlots[0]->GetOffset();
 
 	if(bIsCopying)
 	{
@@ -1545,13 +1562,13 @@ void SDataprepGraphTrackWidget::UpdateLayoutOnDrag()
 		}
 	}
 
-	const FMargin DragSlotOffset = DragSlot->OffsetAttr.Get(FMargin());
-	const FMargin HoveredDropSlotOffset = DropSlots[IndexHovered]->OffsetAttr.Get(FMargin());
+	const FMargin DragSlotOffset = DragSlot->GetOffset();
+	const FMargin HoveredDropSlotOffset = DropSlots[IndexHovered]->GetOffset();
 
 	if(IndexHovered > 0 && DragSlotOffset.Left < HoveredDropSlotOffset.Left)
 	{
 		SConstraintCanvas::FSlot* OverlappedSlot = ActionSlots[IndexHovered-1];
-		FMargin OverlappedSlotOffset = OverlappedSlot->OffsetAttr.Get(FMargin());
+		FMargin OverlappedSlotOffset = OverlappedSlot->GetOffset();
 
 		const float OverlappedActionSlotLeft = HoveredDropSlotOffset.Left - OverlappedSlotOffset.Right;
 		const float Delta =  HoveredDropSlotOffset.Left - DragSlotOffset.Left;
@@ -1562,14 +1579,14 @@ void SDataprepGraphTrackWidget::UpdateLayoutOnDrag()
 		if(NewLeftOffset < MaximalLeftOffset)
 		{
 			OverlappedSlotOffset.Left = NewLeftOffset;
-			OverlappedSlot->OffsetAttr.Set(OverlappedSlotOffset);
+			OverlappedSlot->SetOffset(OverlappedSlotOffset);
 		}
 	}
 	else if(IndexHovered < (SlotCount-1) && DragSlotOffset.Left > (HoveredDropSlotOffset.Left + HoveredDropSlotOffset.Right) )
 	{
 		SConstraintCanvas::FSlot* OverlappedSlot = ActionSlots[IndexHovered+1];
-		FMargin OverlappedSlotOffset = OverlappedSlot->OffsetAttr.Get(FMargin());
-		FMargin OverlappedDropSlotOffset = DropSlots[IndexHovered+1]->OffsetAttr.Get(FMargin());
+		FMargin OverlappedSlotOffset = OverlappedSlot->GetOffset();
+		FMargin OverlappedDropSlotOffset = DropSlots[IndexHovered+1]->GetOffset();
 
 		const float OverlappedActionSlotLeft = OverlappedDropSlotOffset.Left + OverlappedDropSlotOffset.Right;
 		const float MinimalLeftOffset = HoveredDropSlotOffset.Left + HoveredDropSlotOffset.Right;
@@ -1579,7 +1596,7 @@ void SDataprepGraphTrackWidget::UpdateLayoutOnDrag()
 		if(NewLeftOffset > MinimalLeftOffset)
 		{
 			OverlappedSlotOffset.Left = NewLeftOffset;
-			OverlappedSlot->OffsetAttr.Set(OverlappedSlotOffset);
+			OverlappedSlot->SetOffset(OverlappedSlotOffset);
 		}
 	}
 
@@ -1603,14 +1620,14 @@ void SDataprepGraphTrackWidget::OnStartNodeDrag(int32 InIndexDragged)
 	IndexHovered = InIndexDragged;
 	DragIndicatorIndex = INDEX_NONE;
 
-	FMargin Offset = ActionSlots[IndexDragged]->OffsetAttr.Get(FMargin());
-	DragSlot->OffsetAttr.Set(Offset);
+	FMargin Offset = ActionSlots[IndexDragged]->GetOffset();
+	DragSlot->SetOffset(Offset);
 	DragSlot->AttachWidget(ActionNodes[IndexDragged]);
 
 	const FVector2D LocalCursorPosition = GetTickSpaceGeometry().AbsoluteToLocal(FSlateApplication::Get().GetCursorPos());
 	TrackCursorOffset = FVector2D(LocalCursorPosition.X - Offset.Left, LocalCursorPosition.Y);
 
-	AbscissaRange.X = ActionSlots[0]->OffsetAttr.Get(FMargin()).Left;
+	AbscissaRange.X = ActionSlots[0]->GetOffset().Left;
 	// AbscissaRange.Y will be properly update in the call to OnControlKeyDown
 	AbscissaRange.Y = AbscissaRange.X;
 
@@ -1631,7 +1648,7 @@ void SDataprepGraphTrackWidget::OnNodeDragged(float DragDelta, DragCallback Comp
 		return;
 	}
 
-	FMargin DragSlotOffset = DragSlot->OffsetAttr.Get(FMargin());
+	FMargin DragSlotOffset = DragSlot->GetOffset();
 
 	const float DesiredAbscissa = DragSlotOffset.Left + DragDelta;
 	 
@@ -1666,7 +1683,7 @@ void SDataprepGraphTrackWidget::OnNodeDragged(float DragDelta, DragCallback Comp
 
 			// Update drag slot
 			DragSlotOffset.Left = NodeNewAbscissa;
-			DragSlot->OffsetAttr.Set(DragSlotOffset);
+			DragSlot->SetOffset(DragSlotOffset);
 
 			// Check if center of dragged widget is over a neighboring widget by at least half its size
 			IndexHovered = GetHoveredActionNode(DragSlotOffset.Left, DragSlotOffset.Right);
@@ -1699,7 +1716,7 @@ void SDataprepGraphTrackWidget::OnEndNodeDrag(FVector2D& OutScreenSpacePosition,
 
 	SlotCount = ActionNodes.Num();
 
-	DragSlot->OffsetAttr.Set(FVector2D(-100., -10.f));
+	DragSlot->SetOffset(FVector2D(-100., -10.f));
 	DragSlot->AttachWidget(SNullWidget::NullWidget);
 
 	DropSlots[SlotCount]->AttachWidget(SNullWidget::NullWidget);
@@ -1721,7 +1738,7 @@ void SDataprepGraphTrackWidget::UpdateDragIndicator(FVector2D MousePosition)
 	int32 NewDragIndicatorIndex = INDEX_NONE;
 
 	const float LocalAbscissa = GetTickSpaceGeometry().AbsoluteToLocal(MousePosition).X - 1.f;
-	FMargin SlotOffset = ActionSlots[0]->OffsetAttr.Get(FMargin());
+	FMargin SlotOffset = ActionSlots[0]->GetOffset();
 	float MaxAbscissa = SlotOffset.Left + (SlotOffset.Right * 0.25f);
 
 	if(LocalAbscissa < MaxAbscissa)
@@ -1732,7 +1749,7 @@ void SDataprepGraphTrackWidget::UpdateDragIndicator(FVector2D MousePosition)
 	{
 		for(int32 Index = 0; Index <= SlotCount; ++Index)
 		{
-			SlotOffset = ActionSlots[Index]->OffsetAttr.Get(FMargin());
+			SlotOffset = ActionSlots[Index]->GetOffset();
 			if(LocalAbscissa < (SlotOffset.Left + SlotOffset.Right + 1.f))
 			{
 				NewDragIndicatorIndex = Index;
@@ -1759,18 +1776,18 @@ void SDataprepGraphTrackWidget::UpdateDragIndicator(FVector2D MousePosition)
 
 	if(NewDragIndicatorIndex == SlotCount)
 	{
-		float TrackLeftOffset = ActionSlots[SlotCount]->OffsetAttr.Get(FMargin()).Left + (SDataprepGraphActionNode::DefaultWidth * 0.5f) - 16.f;
+		float TrackLeftOffset = ActionSlots[SlotCount]->GetOffset().Left + (SDataprepGraphActionNode::DefaultWidth * 0.5f) - 16.f;
 
-		SlotOffset = TrackSlots[0]->OffsetAttr.Get(FMargin());
+		SlotOffset = TrackSlots[0]->GetOffset();
 
-		TrackSlot->OffsetAttr.Set(FMargin(TrackLeftOffset, SlotOffset.Top, SlotOffset.Right, SlotOffset.Bottom));
+		TrackSlot->SetOffset(FMargin(TrackLeftOffset, SlotOffset.Top, SlotOffset.Right, SlotOffset.Bottom));
 		TrackSlot->AttachWidget(TrackSlotDrop.ToSharedRef());
 
 		UpdateLine(SlotCount);
 	}
 	else
 	{
-		TrackSlot->OffsetAttr.Set(TrackSlots[DragIndicatorIndex]->OffsetAttr.Get(FMargin()));
+		TrackSlot->SetOffset(TrackSlots[DragIndicatorIndex]->GetOffset());
 		TrackSlot->AttachWidget(TrackSlotDrop.ToSharedRef());
 	}
 }
@@ -1794,7 +1811,7 @@ void SDataprepGraphTrackWidget::OnControlKeyDown(bool bCtrlDown, FVector2D& OutS
 
 	SlotCount = ActionNodes.Num();
 
-	FMargin DragSlotOffset = DragSlot->OffsetAttr.Get(FMargin());
+	FMargin DragSlotOffset = DragSlot->GetOffset();
 	IndexHovered = GetHoveredActionNode(DragSlotOffset.Left, DragSlotOffset.Right);
 
 	for(int32 Index = 0; Index < SlotCount; ++Index)
@@ -1810,7 +1827,7 @@ void SDataprepGraphTrackWidget::OnControlKeyDown(bool bCtrlDown, FVector2D& OutS
 
 		// Set the maximum move on the right and new working size
 		const FVector2D DragSize = GetNodeSize(DragSlot->GetWidget());
-		const FMargin SlotOffset = TrackSlots[SlotCount]->OffsetAttr.Get(FMargin());
+		const FMargin SlotOffset = TrackSlots[SlotCount]->GetOffset();
 
 		AbscissaRange.Y = SlotOffset.Left - (DragSize.X * 0.5f) + InterSlotSpacing + 1.f;
 		WorkingSize.X = AbscissaRange.Y + DragSize.X + 1.f;
@@ -1819,17 +1836,17 @@ void SDataprepGraphTrackWidget::OnControlKeyDown(bool bCtrlDown, FVector2D& OutS
 	}
 	else
 	{
-		FMargin SlotOffset = DropSlots[SlotCount - 1]->OffsetAttr.Get(FMargin());
+		FMargin SlotOffset = DropSlots[SlotCount - 1]->GetOffset();
 		AbscissaRange.Y = SlotOffset.Left + SlotOffset.Right;
 
-		WorkingSize.X = DropSlots[SlotCount]->OffsetAttr.Get(FMargin()).Left + SDataprepGraphActionNode::DefaultWidth + InterSlotSpacing;
+		WorkingSize.X = DropSlots[SlotCount]->GetOffset().Left + SDataprepGraphActionNode::DefaultWidth + InterSlotSpacing;
 
 		if(DragSlotOffset.Left > AbscissaRange.Y)
 		{
 			OutScreenSpacePosition = GetAbsoluteSoftwareCursorPosition();
 
 			DragSlotOffset.Left = AbscissaRange.Y;
-			DragSlot->OffsetAttr.Set(DragSlotOffset);
+			DragSlot->SetOffset(DragSlotOffset);
 
 			IndexHovered = GetHoveredActionNode(DragSlotOffset.Left, DragSlotOffset.Right);
 		}
@@ -1866,14 +1883,14 @@ void SDataprepGraphTrackWidget::UpdateLine(int32 LastIndex)
 
 	if(SlotCount > 1)
 	{
-		FMargin OffsetDiff = TrackSlots[LastIndex]->OffsetAttr.Get(FMargin()) - TrackSlots[0]->OffsetAttr.Get(FMargin());
-		FMargin LineOffset = LineSlot->OffsetAttr.Get(FMargin());
+		FMargin OffsetDiff = TrackSlots[LastIndex]->GetOffset() - TrackSlots[0]->GetOffset();
+		FMargin LineOffset = LineSlot->GetOffset();
 
-		LineOffset.Left = TrackSlots[0]->OffsetAttr.Get(FMargin()).Left + 16.f;
+		LineOffset.Left = TrackSlots[0]->GetOffset().Left + 16.f;
 		LineOffset.Right = OffsetDiff.Left;
 		LineSize.X = OffsetDiff.Left;
 
-		LineSlot->OffsetAttr.Set(LineOffset);
+		LineSlot->SetOffset(LineOffset);
 	}
 }
 
diff --git a/Engine/Plugins/Experimental/DigitalMultiplex/DMXPixelMapping/Source/DMXPixelMappingEditor/Private/Widgets/SDMXPixelMappingDesignerCanvas.cpp b/Engine/Plugins/Experimental/DigitalMultiplex/DMXPixelMapping/Source/DMXPixelMappingEditor/Private/Widgets/SDMXPixelMappingDesignerCanvas.cpp
--- a/Engine/Plugins/Experimental/DigitalMultiplex/DMXPixelMapping/Source/DMXPixelMappingEditor/Private/Widgets/SDMXPixelMappingDesignerCanvas.cpp
+++ b/Engine/Plugins/Experimental/DigitalMultiplex/DMXPixelMapping/Source/DMXPixelMappingEditor/Private/Widgets/SDMXPixelMappingDesignerCanvas.cpp
@@ -29,7 +29,7 @@ void SDMXPixelMappingDesignerCanvas::ArrangeLayeredChildren(const FGeometry& All
 			const SConstraintCanvas::FSlot& CurChild = Children[ChildIndex];
 			FChildZOrder Order;
 			Order.ChildIndex = ChildIndex;
-			Order.ZOrder = CurChild.ZOrderAttr.Get();
+			Order.ZOrder = CurChild.GetZOrder();
 			SlotOrder.Add(Order);
 		}
 		SlotOrder.Sort(FSortSlotsByZOrder());
@@ -43,10 +43,10 @@ void SDMXPixelMappingDesignerCanvas::ArrangeLayeredChildren(const FGeometry& All
 			const EVisibility ChildVisibility = CurWidget->GetVisibility();
 			if (ArrangedChildren.Accepts(ChildVisibility))
 			{
-				const FMargin Offset = CurChild.OffsetAttr.Get();
-				const FVector2D Alignment = CurChild.AlignmentAttr.Get();
-				const FAnchors Anchors = CurChild.AnchorsAttr.Get();
-				const bool AutoSize = CurChild.AutoSizeAttr.Get();
+				const FMargin Offset = CurChild.GetOffset();
+				const FVector2D Alignment = CurChild.GetAlignment();
+				const FAnchors Anchors = CurChild.GetAnchors();
+				const bool AutoSize = CurChild.GetAutoSize();
 				const FMargin AnchorPixels =
 					FMargin(Anchors.Minimum.X * AllottedGeometry.GetLocalSize().X,
 					Anchors.Minimum.Y * AllottedGeometry.GetLocalSize().Y,
diff --git a/Engine/Plugins/Experimental/DigitalMultiplex/DMXPixelMapping/Source/DMXPixelMappingRuntime/Private/Components/DMXPixelMappingFixtureGroupComponent.cpp b/Engine/Plugins/Experimental/DigitalMultiplex/DMXPixelMapping/Source/DMXPixelMappingRuntime/Private/Components/DMXPixelMappingFixtureGroupComponent.cpp
--- a/Engine/Plugins/Experimental/DigitalMultiplex/DMXPixelMapping/Source/DMXPixelMappingRuntime/Private/Components/DMXPixelMappingFixtureGroupComponent.cpp
+++ b/Engine/Plugins/Experimental/DigitalMultiplex/DMXPixelMapping/Source/DMXPixelMappingRuntime/Private/Components/DMXPixelMappingFixtureGroupComponent.cpp
@@ -177,11 +177,11 @@ TSharedRef<SWidget> UDMXPixelMappingFixtureGroupComponent::BuildSlot(TSharedRef<
 			]
 		];
 
-	Slot =
-		&InCanvas->AddSlot()
+	InCanvas->AddSlot()
 		.AutoSize(true)
 		.Alignment(FVector2D::ZeroVector)
 		.ZOrder(ZOrder)
+		.Expose(Slot)
 		[
 			SNew(SOverlay)
 			+ SOverlay::Slot()
@@ -204,7 +204,7 @@ TSharedRef<SWidget> UDMXPixelMappingFixtureGroupComponent::BuildSlot(TSharedRef<
 	Brush.TintColor = GetEditorColor(false);
 	Brush.Margin = FMargin(1.f);
 
-	Slot->Offset(FMargin(PositionX, PositionY, 0.f, 0.f));
+	Slot->SetOffset(FMargin(PositionX, PositionY, 0.f, 0.f));
 	CachedWidget->SetWidthOverride(SizeX);
 	CachedWidget->SetHeightOverride(SizeY);
 
@@ -302,7 +302,7 @@ void UDMXPixelMappingFixtureGroupComponent::SetPositionWithChildren()
 	}, false);
 
 #if WITH_EDITOR
-	Slot->Offset(FMargin(PositionX, PositionY, 0.f, 0.f));
+	Slot->SetOffset(FMargin(PositionX, PositionY, 0.f, 0.f));
 #endif // WITH_EDITOR
 }
 
diff --git a/Engine/Plugins/Experimental/DigitalMultiplex/DMXPixelMapping/Source/DMXPixelMappingRuntime/Private/Components/DMXPixelMappingFixtureGroupItemComponent.cpp b/Engine/Plugins/Experimental/DigitalMultiplex/DMXPixelMapping/Source/DMXPixelMappingRuntime/Private/Components/DMXPixelMappingFixtureGroupItemComponent.cpp
--- a/Engine/Plugins/Experimental/DigitalMultiplex/DMXPixelMapping/Source/DMXPixelMappingRuntime/Private/Components/DMXPixelMappingFixtureGroupItemComponent.cpp
+++ b/Engine/Plugins/Experimental/DigitalMultiplex/DMXPixelMapping/Source/DMXPixelMappingRuntime/Private/Components/DMXPixelMappingFixtureGroupItemComponent.cpp
@@ -203,11 +203,11 @@ TSharedRef<SWidget> UDMXPixelMappingFixtureGroupItemComponent::BuildSlot(TShared
 			]
 		];
 
-	Slot =
-		&InCanvas->AddSlot()
+	InCanvas->AddSlot()
 		.AutoSize(true)
 		.Alignment(FVector2D::ZeroVector)
 		.ZOrder(ZOrder)
+		.Expose(Slot)
 		[
 			SNew(SOverlay)
 			
@@ -237,7 +237,7 @@ TSharedRef<SWidget> UDMXPixelMappingFixtureGroupItemComponent::BuildSlot(TShared
 	Brush.TintColor = GetEditorColor(false);
 	Brush.Margin = FMargin(1.f);
 
-	Slot->Offset(FMargin(PositionX, PositionY, 0.f, 0.f));
+	Slot->SetOffset(FMargin(PositionX, PositionY, 0.f, 0.f));
 	CachedWidget->SetWidthOverride(SizeX);
 	CachedWidget->SetHeightOverride(SizeY);
 	CachedLabelBox->SetWidthOverride(SizeX);
@@ -567,7 +567,7 @@ void UDMXPixelMappingFixtureGroupItemComponent::SetPositionInBoundaryBox(const F
 #if WITH_EDITOR
 		if (Slot != nullptr)
 		{
-			Slot->Offset(FMargin(PositionX, PositionY, 0.f, 0.f));
+			Slot->SetOffset(FMargin(PositionX, PositionY, 0.f, 0.f));
 		}
 
 		RelativePositionX = PositionX - FixtureGroupComponent->GetPosition().X;
@@ -597,7 +597,7 @@ void UDMXPixelMappingFixtureGroupItemComponent::SetPositionFromParent(const FVec
 #if WITH_EDITOR
 	if (Slot != nullptr)
 	{
-		Slot->Offset(FMargin(PositionX, PositionY, 0.f, 0.f));
+		Slot->SetOffset(FMargin(PositionX, PositionY, 0.f, 0.f));
 	}
 #endif // WITH_EDITOR
 }
diff --git a/Engine/Plugins/Experimental/DigitalMultiplex/DMXPixelMapping/Source/DMXPixelMappingRuntime/Private/Components/DMXPixelMappingMatrixCellComponent.cpp b/Engine/Plugins/Experimental/DigitalMultiplex/DMXPixelMapping/Source/DMXPixelMappingRuntime/Private/Components/DMXPixelMappingMatrixCellComponent.cpp
--- a/Engine/Plugins/Experimental/DigitalMultiplex/DMXPixelMapping/Source/DMXPixelMappingRuntime/Private/Components/DMXPixelMappingMatrixCellComponent.cpp
+++ b/Engine/Plugins/Experimental/DigitalMultiplex/DMXPixelMapping/Source/DMXPixelMappingRuntime/Private/Components/DMXPixelMappingMatrixCellComponent.cpp
@@ -99,11 +99,11 @@ TSharedRef<SWidget> UDMXPixelMappingMatrixCellComponent::BuildSlot(TSharedRef<SC
 		.HeightOverride(SizeX)
 		.WidthOverride(SizeY);
 
-	Slot =
-		&InCanvas->AddSlot()
+	InCanvas->AddSlot()
 		.AutoSize(true)
 		.Alignment(FVector2D::ZeroVector)
 		.ZOrder(ZOrder)
+		.Expose(Slot)
 		[
 			CachedWidget.ToSharedRef()
 		];
@@ -113,7 +113,7 @@ TSharedRef<SWidget> UDMXPixelMappingMatrixCellComponent::BuildSlot(TSharedRef<SC
 	Brush.TintColor = GetEditorColor(false);
 	Brush.Margin = FMargin(1.f);
 
-	Slot->Offset(FMargin(PositionX, PositionY, 0.f, 0.f));
+	Slot->SetOffset(FMargin(PositionX, PositionY, 0.f, 0.f));
 	CachedWidget->SetWidthOverride(SizeX);
 	CachedWidget->SetHeightOverride(SizeY);
 
@@ -406,7 +406,7 @@ void UDMXPixelMappingMatrixCellComponent::SetPositionFromParent(const FVector2D&
 #if WITH_EDITOR
 	if (Slot != nullptr)
 	{
-		Slot->Offset(FMargin(PositionX, PositionY, 0.f, 0.f));
+		Slot->SetOffset(FMargin(PositionX, PositionY, 0.f, 0.f));
 	}
 #endif // WITH_EDITOR
 }
@@ -439,7 +439,7 @@ void UDMXPixelMappingMatrixCellComponent::SetPositionInBoundaryBox(const FVector
 #if WITH_EDITOR
 		if (Slot != nullptr)
 		{
-			Slot->Offset(FMargin(PositionX, PositionY, 0.f, 0.f));
+			Slot->SetOffset(FMargin(PositionX, PositionY, 0.f, 0.f));
 		}
 
 		RelativePositionX = PositionX - MatrixComponent->GetPosition().X;
@@ -467,7 +467,7 @@ void UDMXPixelMappingMatrixCellComponent::SetSizeWithinBoundaryBox(const FVector
 #if WITH_EDITOR
 		if (Slot != nullptr)
 		{
-			Slot->Offset(FMargin(PositionX, PositionY, 0.f, 0.f));
+			Slot->SetOffset(FMargin(PositionX, PositionY, 0.f, 0.f));
 		}
 #endif // WITH_EDITOR
 	}
diff --git a/Engine/Plugins/Experimental/DigitalMultiplex/DMXPixelMapping/Source/DMXPixelMappingRuntime/Private/Components/DMXPixelMappingMatrixComponent.cpp b/Engine/Plugins/Experimental/DigitalMultiplex/DMXPixelMapping/Source/DMXPixelMappingRuntime/Private/Components/DMXPixelMappingMatrixComponent.cpp
--- a/Engine/Plugins/Experimental/DigitalMultiplex/DMXPixelMapping/Source/DMXPixelMappingRuntime/Private/Components/DMXPixelMappingMatrixComponent.cpp
+++ b/Engine/Plugins/Experimental/DigitalMultiplex/DMXPixelMapping/Source/DMXPixelMappingRuntime/Private/Components/DMXPixelMappingMatrixComponent.cpp
@@ -222,11 +222,11 @@ TSharedRef<SWidget> UDMXPixelMappingMatrixComponent::BuildSlot(TSharedRef<SConst
 			]
 		];
 
-	Slot =
-		&InCanvas->AddSlot()
+	InCanvas->AddSlot()
 		.AutoSize(true)
 		.Alignment(FVector2D::ZeroVector)
 		.ZOrder(ZOrder)
+		.Expose(Slot)
 		[
 			SNew(SOverlay)
 			+ SOverlay::Slot()
@@ -250,7 +250,7 @@ TSharedRef<SWidget> UDMXPixelMappingMatrixComponent::BuildSlot(TSharedRef<SConst
 
 	Brush.Margin = FMargin(1.f);
 
-	Slot->Offset(FMargin(PositionX, PositionY, 0.f, 0.f));
+	Slot->SetOffset(FMargin(PositionX, PositionY, 0.f, 0.f));
 	CachedWidget->SetWidthOverride(SizeX);
 	CachedWidget->SetHeightOverride(SizeY);
 	CachedLabelBox->SetWidthOverride(SizeX);
@@ -510,7 +510,7 @@ void UDMXPixelMappingMatrixComponent::SetPositionWithChildren()
 	}, false);
 
 #if WITH_EDITOR
-	Slot->Offset(FMargin(PositionX, PositionY, 0.f, 0.f));
+	Slot->SetOffset(FMargin(PositionX, PositionY, 0.f, 0.f));
 #endif // WITH_EDITOR
 }
 
@@ -583,7 +583,7 @@ void UDMXPixelMappingMatrixComponent::SetPositionBasedOnRelativePixel(UDMXPixelM
 	}, false);
 
 #if WITH_EDITOR
-	Slot->Offset(FMargin(PositionX, PositionY, 0.f, 0.f));
+	Slot->SetOffset(FMargin(PositionX, PositionY, 0.f, 0.f));
 #endif // WITH_EDITOR
 
 	PositionXCached = PositionX;
diff --git a/Engine/Plugins/Experimental/DigitalMultiplex/DMXPixelMapping/Source/DMXPixelMappingRuntime/Private/Components/DMXPixelMappingOutputComponent.cpp b/Engine/Plugins/Experimental/DigitalMultiplex/DMXPixelMapping/Source/DMXPixelMappingRuntime/Private/Components/DMXPixelMappingOutputComponent.cpp
--- a/Engine/Plugins/Experimental/DigitalMultiplex/DMXPixelMapping/Source/DMXPixelMappingRuntime/Private/Components/DMXPixelMappingOutputComponent.cpp
+++ b/Engine/Plugins/Experimental/DigitalMultiplex/DMXPixelMapping/Source/DMXPixelMappingRuntime/Private/Components/DMXPixelMappingOutputComponent.cpp
@@ -39,11 +39,11 @@ TSharedRef<SWidget> UDMXPixelMappingOutputComponent::BuildSlot(TSharedRef<SConst
 {
 	ensureMsgf(false, TEXT("You must implement RebuildWidget() in your child class"));
 
-	Slot = 
-		&InCanvas->AddSlot()
+	 InCanvas->AddSlot()
 		.AutoSize(true)
 		.Alignment(FVector2D::ZeroVector)
 		.ZOrder(ZOrder)
+		.Expose(Slot)
 		[
 			CachedWidget.ToSharedRef()
 		];
@@ -68,7 +68,7 @@ void UDMXPixelMappingOutputComponent::SetZOrder(int32 NewZOrder)
 
 	if (Slot != nullptr && CachedWidget.IsValid())
 	{
-		Slot->ZOrder(NewZOrder);
+		Slot->SetZOrder(NewZOrder);
 	}
 }
 
diff --git a/Engine/Plugins/Experimental/DigitalMultiplex/DMXPixelMapping/Source/DMXPixelMappingRuntime/Private/Components/DMXPixelMappingScreenComponent.cpp b/Engine/Plugins/Experimental/DigitalMultiplex/DMXPixelMapping/Source/DMXPixelMappingRuntime/Private/Components/DMXPixelMappingScreenComponent.cpp
--- a/Engine/Plugins/Experimental/DigitalMultiplex/DMXPixelMapping/Source/DMXPixelMappingRuntime/Private/Components/DMXPixelMappingScreenComponent.cpp
+++ b/Engine/Plugins/Experimental/DigitalMultiplex/DMXPixelMapping/Source/DMXPixelMappingRuntime/Private/Components/DMXPixelMappingScreenComponent.cpp
@@ -98,7 +98,7 @@ void UDMXPixelMappingScreenComponent::PostEditChangeChainProperty(FPropertyChang
 		if (PropertyChangedChainEvent.GetPropertyName() == GET_MEMBER_NAME_CHECKED(UDMXPixelMappingScreenComponent, PositionX) ||
 			PropertyChangedChainEvent.GetPropertyName() == GET_MEMBER_NAME_CHECKED(UDMXPixelMappingScreenComponent, PositionY))
 		{
-			Slot->Offset(FMargin(PositionX, PositionY, 0.f, 0.f));
+			Slot->SetOffset(FMargin(PositionX, PositionY, 0.f, 0.f));
 		}
 
 		if (PropertyChangedChainEvent.GetPropertyName() == GET_MEMBER_NAME_CHECKED(UDMXPixelMappingScreenComponent, SizeX) ||
@@ -143,11 +143,11 @@ TSharedRef<SWidget> UDMXPixelMappingScreenComponent::BuildSlot(TSharedRef<SConst
 		.HeightOverride(SizeX)
 		.WidthOverride(SizeY);
 
-	Slot = 
-		&InCanvas->AddSlot()
+	InCanvas->AddSlot()
 		.AutoSize(true)
 		.Alignment(FVector2D::ZeroVector)
 		.ZOrder(ZOrder)
+		.Expose(Slot)
 		[
 			CachedWidget.ToSharedRef()
 		];
@@ -157,7 +157,7 @@ TSharedRef<SWidget> UDMXPixelMappingScreenComponent::BuildSlot(TSharedRef<SConst
 	Brush.TintColor = GetEditorColor(false);
 	Brush.Margin = FMargin(1.f);
 
-	Slot->Offset(FMargin(PositionX, PositionY, 0.f, 0.f));
+	Slot->SetOffset(FMargin(PositionX, PositionY, 0.f, 0.f));
 	CachedWidget->SetWidthOverride(SizeX);
 	CachedWidget->SetHeightOverride(SizeY);
 
@@ -561,7 +561,7 @@ void UDMXPixelMappingScreenComponent::SetPosition(const FVector2D& InPosition)
 #if WITH_EDITOR
 	if (Slot != nullptr)
 	{
-		Slot->Offset(FMargin(PositionX, PositionY, 0.f, 0.f));
+		Slot->SetOffset(FMargin(PositionX, PositionY, 0.f, 0.f));
 	}
 #endif // WITH_EDITOR
 }
diff --git a/Engine/Source/Runtime/Slate/Private/Widgets/Layout/SConstraintCanvas.cpp b/Engine/Source/Runtime/Slate/Private/Widgets/Layout/SConstraintCanvas.cpp
--- a/Engine/Source/Runtime/Slate/Private/Widgets/Layout/SConstraintCanvas.cpp
+++ b/Engine/Source/Runtime/Slate/Private/Widgets/Layout/SConstraintCanvas.cpp
@@ -5,6 +5,32 @@
 #include "Layout/ArrangedChildren.h"
 #include "SlateSettings.h"
 
+PRAGMA_DISABLE_DEPRECATION_WARNINGS
+void SConstraintCanvas::FSlot::Construct(const FChildren& SlotOwner, FSlotArguments&& InArgs)
+{
+	TSlotBase<FSlot>::Construct(SlotOwner, MoveTemp(InArgs));
+	if (InArgs._Offset.IsSet())
+	{
+		OffsetAttr = MoveTemp(InArgs._Offset);
+	}
+	if (InArgs._Anchors.IsSet())
+	{
+		AnchorsAttr = MoveTemp(InArgs._Anchors);
+	}
+	if (InArgs._Alignment.IsSet())
+	{
+		AlignmentAttr = MoveTemp(InArgs._Alignment);
+	}
+	if (InArgs._AutoSize.IsSet())
+	{
+		AutoSizeAttr = MoveTemp(InArgs._AutoSize);
+	}
+	if (InArgs._ZOrder.IsSet())
+	{
+		ZOrderAttr = MoveTemp(InArgs._ZOrder);
+	}
+}
+PRAGMA_ENABLE_DEPRECATION_WARNINGS
 
 /* SConstraintCanvas interface
  *****************************************************************************/
@@ -18,20 +44,22 @@ SConstraintCanvas::SConstraintCanvas()
 
 void SConstraintCanvas::Construct( const SConstraintCanvas::FArguments& InArgs )
 {
-	const int32 NumSlots = InArgs.Slots.Num();
-	for ( int32 SlotIndex = 0; SlotIndex < NumSlots; ++SlotIndex )
-	{
-		Children.Add( InArgs.Slots[SlotIndex] );
-	}
+	Children.AddSlots(MoveTemp(const_cast<TArray<FSlot::FSlotArguments>&>(InArgs._Slots)));
+}
+
+SConstraintCanvas::FSlot::FSlotArguments SConstraintCanvas::Slot()
+{
+	return FSlot::FSlotArguments(MakeUnique<FSlot>());
+}
+
+SConstraintCanvas::FScopedWidgetSlotArguments SConstraintCanvas::AddSlot()
+{
+	return FScopedWidgetSlotArguments { MakeUnique<FSlot>(), this->Children, INDEX_NONE };
 }
 
 void SConstraintCanvas::ClearChildren()
 {
-	if ( Children.Num() )
-	{
-		Invalidate(EInvalidateWidget::Layout);
-		Children.Empty();
-	}
+	Children.Empty();
 }
 
 int32 SConstraintCanvas::RemoveSlot( const TSharedRef<SWidget>& SlotWidget )
@@ -95,7 +123,7 @@ void SConstraintCanvas::ArrangeLayeredChildren(const FGeometry& AllottedGeometry
 
 			FChildZOrder Order;
 			Order.ChildIndex = ChildIndex;
-			Order.ZOrder = CurChild.ZOrderAttr.Get();
+			Order.ZOrder = CurChild.GetZOrder();
 			SlotOrder.Add(Order);
 		}
 
@@ -112,11 +140,11 @@ void SConstraintCanvas::ArrangeLayeredChildren(const FGeometry& AllottedGeometry
 			const EVisibility ChildVisibility = CurWidget->GetVisibility();
 			if (ArrangedChildren.Accepts(ChildVisibility))
 			{
-				const FMargin Offset = CurChild.OffsetAttr.Get();
-				const FVector2D Alignment = CurChild.AlignmentAttr.Get();
-				const FAnchors Anchors = CurChild.AnchorsAttr.Get();
+				const FMargin Offset = CurChild.GetOffset();
+				const FVector2D Alignment = CurChild.GetAlignment();
+				const FAnchors Anchors = CurChild.GetAnchors();
 
-				const bool AutoSize = CurChild.AutoSizeAttr.Get();
+				const bool AutoSize = CurChild.GetAutoSize();
 
 				const FMargin AnchorPixels =
 					FMargin(Anchors.Minimum.X * AllottedGeometry.GetLocalSize().X,
@@ -248,13 +276,13 @@ FVector2D SConstraintCanvas::ComputeDesiredSize( float ) const
 		// As long as the widgets are not collapsed, they should contribute to the desired size.
 		if ( ChildVisibilty != EVisibility::Collapsed )
 		{
-			const FMargin Offset = CurChild.OffsetAttr.Get();
-			const FVector2D Alignment = CurChild.AlignmentAttr.Get();
-			const FAnchors Anchors = CurChild.AnchorsAttr.Get();
+			const FMargin Offset = CurChild.GetOffset();
+			const FVector2D Alignment = CurChild.GetAlignment();
+			const FAnchors Anchors = CurChild.GetAnchors();
 
 			const FVector2D SlotSize = FVector2D(Offset.Right, Offset.Bottom);
 
-			const bool AutoSize = CurChild.AutoSizeAttr.Get();
+			const bool AutoSize = CurChild.GetAutoSize();
 
 			const FVector2D Size = AutoSize ? Widget->GetDesiredSize() : SlotSize;
 
diff --git a/Engine/Source/Runtime/Slate/Public/Widgets/Layout/SConstraintCanvas.h b/Engine/Source/Runtime/Slate/Public/Widgets/Layout/SConstraintCanvas.h
--- a/Engine/Source/Runtime/Slate/Public/Widgets/Layout/SConstraintCanvas.h
+++ b/Engine/Source/Runtime/Slate/Public/Widgets/Layout/SConstraintCanvas.h
@@ -29,81 +29,113 @@ public:
 	/**
 	 * ConstraintCanvas slots allow child widgets to be positioned and sized
 	 */
-	class FSlot : public TSlotBase<FSlot>
+	PRAGMA_DISABLE_DEPRECATION_WARNINGS
+	class SLATE_API FSlot : public TSlotBase<FSlot>
 	{
-	public:		
-		FSlot& Offset( const TAttribute<FMargin>& InOffset )
+	public:
+		/** Default values for a slot. */
+		FSlot()
+			: TSlotBase<FSlot>()
+			, OffsetAttr(FMargin(0, 0, 1, 1))
+			, AnchorsAttr(FAnchors(0.0f, 0.0f))
+			, AlignmentAttr(FVector2D(0.5f, 0.5f))
+			, AutoSizeAttr(false)
+			, ZOrderAttr(0)
+		{ }
+
+		SLATE_SLOT_BEGIN_ARGS(FSlot, TSlotBase<FSlot>)
+			SLATE_ATTRIBUTE(FMargin, Offset)
+			SLATE_ATTRIBUTE(FAnchors, Anchors)
+			SLATE_ATTRIBUTE(FVector2D, Alignment)
+			SLATE_ATTRIBUTE(bool, AutoSize)
+			SLATE_ATTRIBUTE(float, ZOrder)
+		SLATE_SLOT_END_ARGS()
+
+		void Construct(const FChildren& SlotOwner, FSlotArguments&& InArgs);
+
+		void SetOffset( const TAttribute<FMargin>& InOffset )
 		{
 			SetAttribute(OffsetAttr, InOffset, EInvalidateWidgetReason::Layout);
-			return *this;
 		}
 
-		FSlot& Anchors( const TAttribute<FAnchors>& InAnchors )
+		FMargin GetOffset() const
+		{
+			return OffsetAttr.Get();
+		}
+
+		void SetAnchors( const TAttribute<FAnchors>& InAnchors )
 		{
 			SetAttribute(AnchorsAttr, InAnchors, EInvalidateWidgetReason::Layout);
-			return *this;
 		}
 
-		FSlot& Alignment(const TAttribute<FVector2D>& InAlignment)
+		FAnchors GetAnchors() const
+		{
+			return AnchorsAttr.Get();
+		}
+
+		void SetAlignment(const TAttribute<FVector2D>& InAlignment)
 		{
 			SetAttribute(AlignmentAttr, InAlignment, EInvalidateWidgetReason::Layout);
-			return *this;
 		}
 
-		FSlot& AutoSize(const TAttribute<bool>& InAutoSize)
+		FVector2D GetAlignment() const
+		{
+			return AlignmentAttr.Get();
+		}
+
+		void SetAutoSize(const TAttribute<bool>& InAutoSize)
 		{
 			SetAttribute(AutoSizeAttr, InAutoSize, EInvalidateWidgetReason::Layout);
-			return *this;
 		}
 
-		FSlot& ZOrder(const TAttribute<float>& InZOrder)
+		bool GetAutoSize() const
+		{
+			return AutoSizeAttr.Get();
+		}
+
+		void SetZOrder(const TAttribute<float>& InZOrder)
 		{
 			// Layout isn't entirely correct here, but Paint wouldn't be either.
 			// We need the parent to redraw the children because the logical order didn't change
 			// but the paint order may have if ZOrder changed, so the painted elements need to
 			// be resorted.
 			SetAttribute(ZOrderAttr, InZOrder, EInvalidateWidgetReason::Layout);
-			return *this;
 		}
 
-		FSlot& Expose( FSlot*& OutVarToInit )
+		float GetZOrder() const
 		{
-			OutVarToInit = this;
-			return *this;
+			return ZOrderAttr.Get();
 		}
 
+	public:
 		/** Offset */
+		UE_DEPRECATED(5.0, "Direct access to OffsetAttr is now deprecated. Use the getter or setter.")
 		TAttribute<FMargin> OffsetAttr;
 
 		/** Anchors */
+		UE_DEPRECATED(5.0, "Direct access to AnchorsAttr is now deprecated. Use the getter or setter.")
 		TAttribute<FAnchors> AnchorsAttr;
 
 		/** Size */
+		UE_DEPRECATED(5.0, "Direct access to AlignmentAttr is now deprecated. Use the getter or setter.")
 		TAttribute<FVector2D> AlignmentAttr;
 
 		/** Auto-Size */
+		UE_DEPRECATED(5.0, "Direct access to AutoSizeAttr is now deprecated. Use the getter or setter.")
 		TAttribute<bool> AutoSizeAttr;
 
 		/** Z-Order */
+		UE_DEPRECATED(5.0, "Direct access to ZOrderAttr is now deprecated. Use the getter or setter.")
 		TAttribute<float> ZOrderAttr;
-
-		/** Default values for a slot. */
-		FSlot()
-			: TSlotBase<FSlot>()
-			, OffsetAttr( FMargin( 0, 0, 1, 1 ) )
-			, AnchorsAttr( FAnchors( 0.0f, 0.0f ) )
-			, AlignmentAttr( FVector2D( 0.5f, 0.5f ) )
-			, AutoSizeAttr( false )
-			, ZOrderAttr( 0 )
-		{ }
 	};
+	PRAGMA_ENABLE_DEPRECATION_WARNINGS
 
 	SLATE_BEGIN_ARGS( SConstraintCanvas )
 		{
 			_Visibility = EVisibility::SelfHitTestInvisible;
 		}
 
-		SLATE_SUPPORTS_SLOT( SConstraintCanvas::FSlot )
+		SLATE_SLOT_ARGUMENT( FSlot, Slots)
 
 	SLATE_END_ARGS()
 
@@ -116,24 +148,15 @@ public:
 	 */
 	void Construct( const FArguments& InArgs );
 
-	static FSlot& Slot()
-	{
-		return *(new FSlot());
-	}
+	static FSlot::FSlotArguments Slot();
 
+	using FScopedWidgetSlotArguments = TPanelChildren<FSlot>::FScopedWidgetSlotArguments;
 	/**
 	 * Adds a content slot.
 	 *
 	 * @return The added slot.
 	 */
-	FSlot& AddSlot()
-	{
-		Invalidate(EInvalidateWidget::Layout);
-
-		SConstraintCanvas::FSlot& NewSlot = *(new FSlot());
-		this->Children.Add( &NewSlot );
-		return NewSlot;
-	}
+	FScopedWidgetSlotArguments AddSlot();
 
 	/**
 	 * Removes a particular content slot.
diff --git a/Engine/Source/Runtime/UMG/Private/Components/CanvasPanelSlot.cpp b/Engine/Source/Runtime/UMG/Private/Components/CanvasPanelSlot.cpp
--- a/Engine/Source/Runtime/UMG/Private/Components/CanvasPanelSlot.cpp
+++ b/Engine/Source/Runtime/UMG/Private/Components/CanvasPanelSlot.cpp
@@ -140,9 +140,9 @@ void UCanvasPanelSlot::SetLayout(const FAnchorData& InLayoutData)
 
 	if ( Slot )
 	{
-		Slot->Offset(LayoutData.Offsets);
-		Slot->Anchors(LayoutData.Anchors);
-		Slot->Alignment(LayoutData.Alignment);
+		Slot->SetOffset(LayoutData.Offsets);
+		Slot->SetAnchors(LayoutData.Anchors);
+		Slot->SetAlignment(LayoutData.Alignment);
 	}
 }
 
@@ -158,7 +158,7 @@ void UCanvasPanelSlot::SetPosition(FVector2D InPosition)
 
 	if ( Slot )
 	{
-		Slot->Offset(LayoutData.Offsets);
+		Slot->SetOffset(LayoutData.Offsets);
 	}
 }
 
@@ -166,7 +166,7 @@ FVector2D UCanvasPanelSlot::GetPosition() const
 {
 	if ( Slot )
 	{
-		FMargin Offsets = Slot->OffsetAttr.Get();
+		FMargin Offsets = Slot->GetOffset();
 		return FVector2D(Offsets.Left, Offsets.Top);
 	}
 
@@ -180,7 +180,7 @@ void UCanvasPanelSlot::SetSize(FVector2D InSize)
 
 	if ( Slot )
 	{
-		Slot->Offset(LayoutData.Offsets);
+		Slot->SetOffset(LayoutData.Offsets);
 	}
 }
 
@@ -188,7 +188,7 @@ FVector2D UCanvasPanelSlot::GetSize() const
 {
 	if ( Slot )
 	{
-		FMargin Offsets = Slot->OffsetAttr.Get();
+		FMargin Offsets = Slot->GetOffset();
 		return FVector2D(Offsets.Right, Offsets.Bottom);
 	}
 
@@ -200,7 +200,7 @@ void UCanvasPanelSlot::SetOffsets(FMargin InOffset)
 	LayoutData.Offsets = InOffset;
 	if ( Slot )
 	{
-		Slot->Offset(InOffset);
+		Slot->SetOffset(InOffset);
 	}
 }
 
@@ -208,7 +208,7 @@ FMargin UCanvasPanelSlot::GetOffsets() const
 {
 	if ( Slot )
 	{
-		return Slot->OffsetAttr.Get();
+		return Slot->GetOffset();
 	}
 
 	return LayoutData.Offsets;
@@ -219,7 +219,7 @@ void UCanvasPanelSlot::SetAnchors(FAnchors InAnchors)
 	LayoutData.Anchors = InAnchors;
 	if ( Slot )
 	{
-		Slot->Anchors(InAnchors);
+		Slot->SetAnchors(InAnchors);
 	}
 }
 
@@ -227,7 +227,7 @@ FAnchors UCanvasPanelSlot::GetAnchors() const
 {
 	if ( Slot )
 	{
-		return Slot->AnchorsAttr.Get();
+		return Slot->GetAnchors();
 	}
 
 	return LayoutData.Anchors;
@@ -238,7 +238,7 @@ void UCanvasPanelSlot::SetAlignment(FVector2D InAlignment)
 	LayoutData.Alignment = InAlignment;
 	if ( Slot )
 	{
-		Slot->Alignment(InAlignment);
+		Slot->SetAlignment(InAlignment);
 	}
 }
 
@@ -246,7 +246,7 @@ FVector2D UCanvasPanelSlot::GetAlignment() const
 {
 	if ( Slot )
 	{
-		return Slot->AlignmentAttr.Get();
+		return Slot->GetAlignment();
 	}
 
 	return LayoutData.Alignment;
@@ -257,7 +257,7 @@ void UCanvasPanelSlot::SetAutoSize(bool InbAutoSize)
 	bAutoSize = InbAutoSize;
 	if ( Slot )
 	{
-		Slot->AutoSize(InbAutoSize);
+		Slot->SetAutoSize(InbAutoSize);
 	}
 }
 
@@ -265,7 +265,7 @@ bool UCanvasPanelSlot::GetAutoSize() const
 {
 	if ( Slot )
 	{
-		return Slot->AutoSizeAttr.Get();
+		return Slot->GetAutoSize();
 	}
 
 	return bAutoSize;
@@ -276,7 +276,7 @@ void UCanvasPanelSlot::SetZOrder(int32 InZOrder)
 	ZOrder = InZOrder;
 	if ( Slot )
 	{
-		Slot->ZOrder(InZOrder);
+		Slot->SetZOrder(InZOrder);
 	}
 }
 
@@ -284,7 +284,7 @@ int32 UCanvasPanelSlot::GetZOrder() const
 {
 	if ( Slot )
 	{
-		return Slot->ZOrderAttr.Get();
+		return Slot->GetZOrder();
 	}
 
 	return ZOrder;
@@ -295,7 +295,7 @@ void UCanvasPanelSlot::SetMinimum(FVector2D InMinimumAnchors)
 	LayoutData.Anchors.Minimum = InMinimumAnchors;
 	if ( Slot )
 	{
-		Slot->Anchors(LayoutData.Anchors);
+		Slot->SetAnchors(LayoutData.Anchors);
 	}
 }
 
@@ -304,7 +304,7 @@ void UCanvasPanelSlot::SetMaximum(FVector2D InMaximumAnchors)
 	LayoutData.Anchors.Maximum = InMaximumAnchors;
 	if ( Slot )
 	{
-		Slot->Anchors(LayoutData.Anchors);
+		Slot->SetAnchors(LayoutData.Anchors);
 	}
 }
 
diff --git a/Engine/Source/Runtime/UMG/Private/Slate/SWorldWidgetScreenLayer.cpp b/Engine/Source/Runtime/UMG/Private/Slate/SWorldWidgetScreenLayer.cpp
--- a/Engine/Source/Runtime/UMG/Private/Slate/SWorldWidgetScreenLayer.cpp
+++ b/Engine/Source/Runtime/UMG/Private/Slate/SWorldWidgetScreenLayer.cpp
@@ -147,26 +147,26 @@ void SWorldWidgetScreenLayer::Tick(const FGeometry& AllottedGeometry, const doub
 								FVector2D ComponentDrawSize = Entry.WidgetComponent->GetDrawSize();
 								FVector2D ComponentPivot = Entry.WidgetComponent->GetPivot();
 								
-								CanvasSlot->AutoSize(ComponentDrawSize.IsZero() || Entry.WidgetComponent->GetDrawAtDesiredSize());
-								CanvasSlot->Offset(FMargin(LocalPosition.X, LocalPosition.Y, ComponentDrawSize.X, ComponentDrawSize.Y));
-								CanvasSlot->Anchors(FAnchors(0, 0, 0, 0));
-								CanvasSlot->Alignment(ComponentPivot);
+								CanvasSlot->SetAutoSize(ComponentDrawSize.IsZero() || Entry.WidgetComponent->GetDrawAtDesiredSize());
+								CanvasSlot->SetOffset(FMargin(LocalPosition.X, LocalPosition.Y, ComponentDrawSize.X, ComponentDrawSize.Y));
+								CanvasSlot->SetAnchors(FAnchors(0, 0, 0, 0));
+								CanvasSlot->SetAlignment(ComponentPivot);
 								
 								if (GSlateWorldWidgetZOrder != 0)
 								{
-									CanvasSlot->ZOrder(-ViewportPosition.Z);
+									CanvasSlot->SetZOrder(-ViewportPosition.Z);
 								}
 							}
 							else
 							{
-								CanvasSlot->AutoSize(DrawSize.IsZero());
-								CanvasSlot->Offset(FMargin(LocalPosition.X, LocalPosition.Y, DrawSize.X, DrawSize.Y));
-								CanvasSlot->Anchors(FAnchors(0, 0, 0, 0));
-								CanvasSlot->Alignment(Pivot);
+								CanvasSlot->SetAutoSize(DrawSize.IsZero());
+								CanvasSlot->SetOffset(FMargin(LocalPosition.X, LocalPosition.Y, DrawSize.X, DrawSize.Y));
+								CanvasSlot->SetAnchors(FAnchors(0, 0, 0, 0));
+								CanvasSlot->SetAlignment(Pivot);
 
 								if (GSlateWorldWidgetZOrder != 0)
 								{
-									CanvasSlot->ZOrder(-ViewportPosition.Z);
+									CanvasSlot->SetZOrder(-ViewportPosition.Z);
 								}
 							}
 						}
