Dataprep: Actors' labels are not coming in the visual dataprep for CAD imports

diff --git a/Engine/Plugins/Enterprise/DataprepEditor/Source/DataprepEditor/Private/DataprepEditorOutlinerMode.h b/Engine/Plugins/Enterprise/DataprepEditor/Source/DataprepEditor/Private/DataprepEditorOutlinerMode.h
--- a/Engine/Plugins/Enterprise/DataprepEditor/Source/DataprepEditor/Private/DataprepEditorOutlinerMode.h
+++ b/Engine/Plugins/Enterprise/DataprepEditor/Source/DataprepEditor/Private/DataprepEditorOutlinerMode.h
@@ -48,7 +48,9 @@ public:
 	FDataprepEditorOutlinerMode(SSceneOutliner* InSceneOutliner, TWeakPtr<FDataprepEditor> InDataprepEditor, TWeakObjectPtr<UWorld> InSpecifiedWorldToDisplay = nullptr)
 		: FActorMode(FActorModeParams(InSceneOutliner, InSpecifiedWorldToDisplay, true, true))
 		, DataprepEditorPtr(InDataprepEditor)
-	{}
+	{
+		FCoreDelegates::OnActorLabelChanged.AddRaw(this, &FDataprepEditorOutlinerMode::OnActorLabelChanged);
+	}
 
 	virtual bool CanRenameItem(const ISceneOutlinerTreeItem& Item) const override { return false; }
 	virtual ESelectionMode::Type GetSelectionMode() const override { return ESelectionMode::Multi; }
@@ -61,6 +63,25 @@ public:
 			DataprepEditor->OnSceneOutlinerSelectionChanged(Item, SelectionType);
 		}
 	}
+
+private:
+	void OnActorLabelChanged(AActor* InChangedActor)
+	{
+		if (!ensure(InChangedActor))
+		{
+			return;
+		}
+
+		if (SceneOutliner && IsActorDisplayable(InChangedActor))
+		{
+			// Force create the item otherwise the outliner may not be notified of a change to the item if it is filtered out
+			if (FSceneOutlinerTreeItemPtr Item = CreateItemFor<FActorTreeItem>(InChangedActor, true))
+			{
+				SceneOutliner->OnItemLabelChanged(Item);
+			}
+		}
+	}
+
 private:
 	TWeakPtr<FDataprepEditor> DataprepEditorPtr;
 };
\ No newline at end of file
