Deprecate the SSCSEditor and other old slate classes that are used with it.

#rnx
#p4v-preflight-copy 17276254

diff --git a/Engine/Plugins/Enterprise/DataprepEditor/Source/DataprepEditor/DataprepEditor.Build.cs b/Engine/Plugins/Enterprise/DataprepEditor/Source/DataprepEditor/DataprepEditor.Build.cs
--- a/Engine/Plugins/Enterprise/DataprepEditor/Source/DataprepEditor/DataprepEditor.Build.cs
+++ b/Engine/Plugins/Enterprise/DataprepEditor/Source/DataprepEditor/DataprepEditor.Build.cs
@@ -44,6 +44,8 @@ namespace UnrealBuildTool.Rules
 					"StatsViewer",
 					"ToolMenus",
 					"UnrealEd",
+					"SubobjectEditor",
+					"SubobjectDataInterface",
 				}
 			);
 
diff --git a/Engine/Plugins/Enterprise/DataprepEditor/Source/DataprepEditor/Private/Widgets/DataprepAssetView.cpp b/Engine/Plugins/Enterprise/DataprepEditor/Source/DataprepEditor/Private/Widgets/DataprepAssetView.cpp
--- a/Engine/Plugins/Enterprise/DataprepEditor/Source/DataprepEditor/Private/Widgets/DataprepAssetView.cpp
+++ b/Engine/Plugins/Enterprise/DataprepEditor/Source/DataprepEditor/Private/Widgets/DataprepAssetView.cpp
@@ -22,7 +22,9 @@
 #include "EditorFontGlyphs.h"
 #include "Modules/ModuleManager.h"
 #include "ISCSEditorUICustomization.h"
-#include "SSCSEditor.h"
+#include "SSubobjectEditor.h"
+#include "SSubobjectInstanceEditor.h"
+#include "SubobjectData.h"
 #include "Widgets/Colors/SColorBlock.h"
 #include "Widgets/Images/SImage.h"
 #include "Widgets/Input/SButton.h"
@@ -92,9 +94,8 @@ void SGraphNodeDetailsWidget::Construct(const FArguments& InArgs)
 			SNew(SBox)
 				.Visibility_Lambda([this]() -> EVisibility { return SelectedActor ? EVisibility::Visible : EVisibility::Collapsed; })
 				[
-					SAssignNew(SCSEditor, SSCSEditor)
-						.EditorMode(EComponentEditorMode::ActorInstance)
-						.ActorContext_Lambda([this]() -> AActor* { return SelectedActor; })
+					SAssignNew(SubobjectEditor, SSubobjectInstanceEditor)
+						.ObjectContext_Lambda([this]() -> UObject* { return SelectedActor; })
 						.OnSelectionUpdated(this, &SGraphNodeDetailsWidget::OnSCSEditorTreeViewSelectionChanged)
 						.AllowEditing(false)
 				]
@@ -104,8 +105,8 @@ void SGraphNodeDetailsWidget::Construct(const FArguments& InArgs)
 			PropertyView.ToSharedRef()
 		];
 
-	SCSEditorUICustomization = MakeShared<FDataprepSCSEditorUICustomization>();
-	SCSEditor->SetUICustomization(SCSEditorUICustomization);
+	SubobjectEditorUICustomization = MakeShared<FDataprepSCSEditorUICustomization>();
+	SubobjectEditor->SetUICustomization(SubobjectEditorUICustomization);
 
 	// Create the border that all of the content will get stuffed into
 	ChildSlot
@@ -177,7 +178,7 @@ void SGraphNodeDetailsWidget::AddPropertiesRecursive(FProperty* Property)
 	}
 }
 
-void SGraphNodeDetailsWidget::OnSCSEditorTreeViewSelectionChanged(const TArray<TSharedPtr<class FSCSEditorTreeNode> >& SelectedNodes)
+void SGraphNodeDetailsWidget::OnSCSEditorTreeViewSelectionChanged(const TArray<FSubobjectEditorTreeNodePtrType>& SelectedNodes)
 {
 	if (SelectedNodes.Num() > 0)
 	{
@@ -185,22 +186,23 @@ void SGraphNodeDetailsWidget::OnSCSEditorTreeViewSelectionChanged(const TArray<T
 		{
 			TArray<UObject*> DetailsObjects;
 
-			for (const TSharedPtr<FSCSEditorTreeNode>& SelectedNode : SelectedNodes)
+			for (const FSubobjectEditorTreeNodePtrType& SelectedNode : SelectedNodes)
 			{
 				if (SelectedNode.IsValid())
 				{
-					if (SelectedNode->GetNodeType() == FSCSEditorTreeNode::RootActorNode)
+					if (SelectedNode->IsRootActorNode())
 					{
 						// Root actor takes precedence
 						DetailsObjects.Empty();
 						DetailsObjects.Add(SelectedActor);
 						break;
 					}
-					else if (SelectedNode->GetNodeType() == FSCSEditorTreeNode::ComponentNode)
+					else if (SelectedNode->IsComponentNode())
 					{
-						if (UActorComponent* ComponentInstance = SelectedNode->FindComponentInstanceInActor(SelectedActor))
+						FSubobjectData* Data = SelectedNode->GetDataSource();
+						if (const UActorComponent* ComponentInstance = Data->FindComponentInstanceInActor(SelectedActor))
 						{
-							DetailsObjects.Add(ComponentInstance);
+							DetailsObjects.Add(const_cast<UActorComponent*>(ComponentInstance));
 						}
 					}
 				}
@@ -358,7 +360,7 @@ void SGraphNodeDetailsWidget::UpdateFromObjects(const TArray<UObject*>& Property
 		}
 	}
 
-	SCSEditorUICustomization->SetHideComponentsTree(SelectedActor == nullptr);
+	SubobjectEditorUICustomization->SetHideComponentsTree(SelectedActor == nullptr);
 
 	PropertyView->SetObjects(SelectionInfo.ObjectsForPropertyEditing);
 
@@ -384,7 +386,7 @@ void SGraphNodeDetailsWidget::UpdateFromObjects(const TArray<UObject*>& Property
 	// Don't update component tree if this update comes from itself (avoids infinite recusrsion)
 	if (!bSelfUpdate)
 	{
-		SCSEditor->UpdateTree();
+		SubobjectEditor->UpdateTree();
 	}
 }
 
diff --git a/Engine/Plugins/Enterprise/DataprepEditor/Source/DataprepEditor/Private/Widgets/DataprepAssetView.h b/Engine/Plugins/Enterprise/DataprepEditor/Source/DataprepEditor/Private/Widgets/DataprepAssetView.h
--- a/Engine/Plugins/Enterprise/DataprepEditor/Source/DataprepEditor/Private/Widgets/DataprepAssetView.h
+++ b/Engine/Plugins/Enterprise/DataprepEditor/Source/DataprepEditor/Private/Widgets/DataprepAssetView.h
@@ -23,6 +23,8 @@
 class IDetailsView;
 class SDataprepAssetView;
 class SDataprepConsumerWidget;
+class SSubobjectEditor;
+class FSubobjectEditorTreeNode;
 
 class SDataprepAssetView : public SCompoundWidget
 {
@@ -93,7 +95,7 @@ private:
 	/** Add this property and all its child properties to SelectedObjectProperties */
 	void AddPropertiesRecursive(FProperty* Property);
 
-	void OnSCSEditorTreeViewSelectionChanged(const TArray<TSharedPtr<class FSCSEditorTreeNode> >& SelectedNodes);
+	void OnSCSEditorTreeViewSelectionChanged(const TArray<TSharedPtr<FSubobjectEditorTreeNode>>& SelectedNodes);
 
 private:
 	/** Property viewing widget */
@@ -112,10 +114,10 @@ private:
 	TSharedPtr<SSplitter> DetailsSplitter;
 
 	/** Component tree */
-	TSharedPtr<class SSCSEditor> SCSEditor;
+	TSharedPtr<SSubobjectEditor> SubobjectEditor;
 
 	/** Customize how the component tree looks like */
-	TSharedPtr<class FDataprepSCSEditorUICustomization> SCSEditorUICustomization;
+	TSharedPtr<class FDataprepSCSEditorUICustomization> SubobjectEditorUICustomization;
 
 	/** The first actor in the currently selected objects */
 	AActor* SelectedActor;
diff --git a/Engine/Source/Editor/Kismet/Public/SCSEditorExtensionContext.h b/Engine/Source/Editor/Kismet/Public/SCSEditorExtensionContext.h
--- a/Engine/Source/Editor/Kismet/Public/SCSEditorExtensionContext.h
+++ b/Engine/Source/Editor/Kismet/Public/SCSEditorExtensionContext.h
@@ -18,19 +18,21 @@ class KISMET_API USCSEditorExtensionContext : public UObject
 	GENERATED_BODY()
 
 public:
-	
+	PRAGMA_DISABLE_DEPRECATION_WARNINGS
 	UE_DEPRECATED(5.0, "GetSCSEditor has been deprecated. Use GetSubobjectEditor instead.")
 	const TWeakPtr<SSCSEditor>& GetSCSEditor() const { return SCSEditor; }
 	
 	const TWeakPtr<SSubobjectEditor>& GetSubobjectEditor() const { return SubobjectEditor; }
 
 private:
+	
 	friend SSCSEditor;
 	friend SSubobjectEditor;
 	friend SSubobjectInstanceEditor;
 	friend SSubobjectBlueprintEditor;
-
+	
 	TWeakPtr<SSCSEditor> SCSEditor;
+	PRAGMA_ENABLE_DEPRECATION_WARNINGS
 	TWeakPtr<SSubobjectEditor> SubobjectEditor;
 
 };
diff --git a/Engine/Source/Editor/Kismet/Public/SSCSEditor.h b/Engine/Source/Editor/Kismet/Public/SSCSEditor.h
--- a/Engine/Source/Editor/Kismet/Public/SSCSEditor.h
+++ b/Engine/Source/Editor/Kismet/Public/SSCSEditor.h
@@ -977,6 +977,8 @@ namespace EComponentEditorMode
 	};
 };
 
+class UE_DEPRECATED(5.0, "SSCSEditor has been deprecated, use a child class of SSubobjectEditor instead.") SSCSEditor;
+
 class KISMET_API SSCSEditor : public SCompoundWidget
 {
 public:
diff --git a/Engine/Source/Editor/Kismet/Public/SSCSEditorMenuContext.h b/Engine/Source/Editor/Kismet/Public/SSCSEditorMenuContext.h
--- a/Engine/Source/Editor/Kismet/Public/SSCSEditorMenuContext.h
+++ b/Engine/Source/Editor/Kismet/Public/SSCSEditorMenuContext.h
@@ -15,9 +15,13 @@ class KISMET_API USSCSEditorMenuContext : public UObject
 {
 	GENERATED_BODY()
 public:
-
+	
+	PRAGMA_DISABLE_DEPRECATION_WARNINGS
+	
 	TWeakPtr<SSCSEditor> SCSEditor;
-
+	
+	PRAGMA_ENABLE_DEPRECATION_WARNINGS
+	
 	TWeakPtr<SSubobjectEditor> SubobjectEditor;
 
 	bool bOnlyShowPasteOption;
