Visual Dataprep crash on commit

diff --git a/Engine/Plugins/Enterprise/DataprepEditor/Source/DataprepEditor/Private/DataprepEditor.h b/Engine/Plugins/Enterprise/DataprepEditor/Source/DataprepEditor/Private/DataprepEditor.h
--- a/Engine/Plugins/Enterprise/DataprepEditor/Source/DataprepEditor/Private/DataprepEditor.h
+++ b/Engine/Plugins/Enterprise/DataprepEditor/Source/DataprepEditor/Private/DataprepEditor.h
@@ -146,6 +146,11 @@ public:
 	/** Returns content folder under which all assets are stored after execution of all producers */
 	FString GetTransientContentFolder();
 
+	UWorld* GetWorld()
+	{
+		return PreviewWorld;
+	}
+
 private:
 	void BindCommands();
 	void OnSaveScene();
diff --git a/Engine/Plugins/Enterprise/DataprepEditor/Source/DataprepEditor/Private/DataprepEditorOutlinerMode.h b/Engine/Plugins/Enterprise/DataprepEditor/Source/DataprepEditor/Private/DataprepEditorOutlinerMode.h
--- a/Engine/Plugins/Enterprise/DataprepEditor/Source/DataprepEditor/Private/DataprepEditorOutlinerMode.h
+++ b/Engine/Plugins/Enterprise/DataprepEditor/Source/DataprepEditor/Private/DataprepEditorOutlinerMode.h
@@ -71,6 +71,16 @@ private:
 		{
 			return;
 		}
+		
+		if (!DataprepEditorPtr.IsValid())
+		{
+			return;
+		}
+
+		if (InChangedActor->GetWorld() != DataprepEditorPtr.Pin()->GetWorld())
+		{
+			return;
+		}
 
 		if (SceneOutliner && IsActorDisplayable(InChangedActor))
 		{
