DataPrep - 3d viewport can't frame selection

diff --git a/Engine/Plugins/Enterprise/DataprepEditor/Source/DataprepEditor/Private/DataprepEditor.cpp b/Engine/Plugins/Enterprise/DataprepEditor/Source/DataprepEditor/Private/DataprepEditor.cpp
--- a/Engine/Plugins/Enterprise/DataprepEditor/Source/DataprepEditor/Private/DataprepEditor.cpp
+++ b/Engine/Plugins/Enterprise/DataprepEditor/Source/DataprepEditor/Private/DataprepEditor.cpp
@@ -1047,6 +1047,9 @@ void FDataprepEditor::CreateTabs()
 
 	// Create Details Panel
 	CreateDetailsViews();
+
+	AssetPreviewView->OnKeyDown().BindSP( SceneViewportView.Get(), &SDataprepEditorViewport::OnKeyDown );
+	ScenePreviewView->OnKeyDown().BindRaw( SceneViewportView.Get(), &SDataprepEditorViewport::OnKeyDown );
 }
 
 TSharedRef<SDockTab> FDataprepEditor::SpawnTabScenePreview(const FSpawnTabArgs & Args)
diff --git a/Engine/Plugins/Enterprise/DataprepEditor/Source/DataprepEditor/Private/DataprepEditor.h b/Engine/Plugins/Enterprise/DataprepEditor/Source/DataprepEditor/Private/DataprepEditor.h
--- a/Engine/Plugins/Enterprise/DataprepEditor/Source/DataprepEditor/Private/DataprepEditor.h
+++ b/Engine/Plugins/Enterprise/DataprepEditor/Source/DataprepEditor/Private/DataprepEditor.h
@@ -3,6 +3,7 @@
 #pragma once
 
 #include "DataprepActionAsset.h"
+#include "DataprepEditorUtils.h"
 #include "DataprepAssetInterface.h"
 #include "DataprepGraph/DataprepGraph.h"
 #include "DataprepStats.h"
@@ -19,6 +20,8 @@
 #include "Toolkits/IToolkitHost.h"
 #include "UObject/GCObject.h"
 #include "UObject/SoftObjectPath.h"
+#include "Widgets/SWidget.h"
+#include "Widgets/DeclarativeSyntaxSupport.h"
 
 class UDataprepParameterizableObject;
 
@@ -51,6 +54,25 @@ namespace AssetPreviewWidget
 	class SAssetsPreviewWidget;
 }
 
+class SDataprepScenePreviewView : public SBorder
+{
+public:
+
+	FDataprepEditorUtils::FOnKeyDown& OnKeyDown() { return OnKeyDownDelegate; }
+
+	virtual FReply OnKeyDown( const FGeometry& MyGeometry, const FKeyEvent& InKeyEvent ) override
+	{
+		if( OnKeyDownDelegate.IsBound() )
+		{
+			return OnKeyDownDelegate.Execute( MyGeometry, InKeyEvent );
+		}
+		return FReply::Unhandled();
+	}
+
+private:
+	FDataprepEditorUtils::FOnKeyDown OnKeyDownDelegate;
+};
+
 /** Tuple linking an asset package path, a unique identifier and the UClass of the asset*/
 typedef TTuple< FString, UClass*, EObjectFlags > FSnapshotDataEntry;
 
@@ -264,7 +286,7 @@ private:
 	TWeakPtr<SDockTab> DetailsTabPtr;
 	TSharedPtr<class SDataprepEditorViewport> SceneViewportView;
 	TSharedPtr<AssetPreviewWidget::SAssetsPreviewWidget> AssetPreviewView;
-	TSharedPtr<SWidget> ScenePreviewView;
+	TSharedPtr<SDataprepScenePreviewView> ScenePreviewView;
 	TSharedPtr<SGraphNodeDetailsWidget> DetailsView;
 	TSharedPtr<class SDataprepAssetView > DataprepAssetView;
 	TSharedPtr<SDataprepGraphEditor> GraphEditor;
diff --git a/Engine/Plugins/Enterprise/DataprepEditor/Source/DataprepEditor/Private/DataprepEditorOutliner.cpp b/Engine/Plugins/Enterprise/DataprepEditor/Source/DataprepEditor/Private/DataprepEditorOutliner.cpp
--- a/Engine/Plugins/Enterprise/DataprepEditor/Source/DataprepEditor/Private/DataprepEditorOutliner.cpp
+++ b/Engine/Plugins/Enterprise/DataprepEditor/Source/DataprepEditor/Private/DataprepEditorOutliner.cpp
@@ -401,7 +401,7 @@ void FDataprepEditor::CreateScenePreviewTab()
 		SceneOutliner->AddColumn(DefaultColumn.Key, DefaultColumn.Value);
 	}
 
-	SAssignNew(ScenePreviewView, SBorder)
+	SAssignNew(ScenePreviewView, SDataprepScenePreviewView)
 		.Padding(2.f)
 		.BorderImage(FEditorStyle::GetBrush("ToolPanel.GroupBorder"))
 		[
diff --git a/Engine/Plugins/Enterprise/DataprepEditor/Source/DataprepEditor/Private/DataprepEditorUtils.h b/Engine/Plugins/Enterprise/DataprepEditor/Source/DataprepEditor/Private/DataprepEditorUtils.h
--- a/Engine/Plugins/Enterprise/DataprepEditor/Source/DataprepEditor/Private/DataprepEditorUtils.h
+++ b/Engine/Plugins/Enterprise/DataprepEditor/Source/DataprepEditor/Private/DataprepEditorUtils.h
@@ -37,6 +37,8 @@ struct FDataprepParametrizationActionData : public FGCObject
 class FDataprepEditorUtils
 {
 public:
+	DECLARE_DELEGATE_RetVal_TwoParams(FReply, FOnKeyDown, const FGeometry& , const FKeyEvent& )
+
 	/**
 	 * Populate a menu builder with the section made for the parameterization
 	 * @param DataprepAsset the asset that own the object
diff --git a/Engine/Plugins/Enterprise/DataprepEditor/Source/DataprepEditor/Private/Widgets/SAssetsPreviewWidget.h b/Engine/Plugins/Enterprise/DataprepEditor/Source/DataprepEditor/Private/Widgets/SAssetsPreviewWidget.h
--- a/Engine/Plugins/Enterprise/DataprepEditor/Source/DataprepEditor/Private/Widgets/SAssetsPreviewWidget.h
+++ b/Engine/Plugins/Enterprise/DataprepEditor/Source/DataprepEditor/Private/Widgets/SAssetsPreviewWidget.h
@@ -4,6 +4,8 @@
 
 #include "CoreMinimal.h"
 
+#include "DataprepEditorUtils.h"
+
 #include "Templates/SharedPointer.h"
 #include "Types/SlateEnums.h"
 #include "Widgets/SCompoundWidget.h"
@@ -75,6 +77,7 @@ namespace AssetPreviewWidget
 
 		FOnSelectionChanged& OnSelectionChanged() { return OnSelectionChangedDelegate; }
 		FOnContextMenu& OnContextMenu() { return OnContextMenuDelegate; }
+		FDataprepEditorUtils::FOnKeyDown& OnKeyDown() { return OnKeyDownDelegate; }
 
 		FText OnGetHighlightText() const;
 
@@ -103,6 +106,15 @@ namespace AssetPreviewWidget
 
 		void SetSelectedAssets(TSet<UObject*> InSelectionSet, ESelectInfo::Type SelectionInfo);
 
+		virtual FReply OnKeyDown( const FGeometry& MyGeometry, const FKeyEvent& InKeyEvent ) override
+		{
+			if( OnKeyDownDelegate.IsBound() )
+			{
+				return OnKeyDownDelegate.Execute( MyGeometry, InKeyEvent );
+			}
+			return FReply::Unhandled();
+		}
+
 	private:
 
 		friend struct FAssetTreeItem;
@@ -165,6 +177,7 @@ namespace AssetPreviewWidget
 
 		FOnSelectionChanged OnSelectionChangedDelegate;
 		FOnContextMenu OnContextMenuDelegate;
+		FDataprepEditorUtils::FOnKeyDown OnKeyDownDelegate;
 
 		FName SortingColumn;
 		EColumnSortMode::Type SortingMode;
