Dataprep: Label name changes after execute in Visual Data Prep

diff --git a/Engine/Plugins/Enterprise/DataprepEditor/Source/DataprepEditor/Private/DataprepEditor.cpp b/Engine/Plugins/Enterprise/DataprepEditor/Source/DataprepEditor/Private/DataprepEditor.cpp
--- a/Engine/Plugins/Enterprise/DataprepEditor/Source/DataprepEditor/Private/DataprepEditor.cpp
+++ b/Engine/Plugins/Enterprise/DataprepEditor/Source/DataprepEditor/Private/DataprepEditor.cpp
@@ -917,6 +917,8 @@ void FDataprepEditor::OnExecutePipeline()
 	// Reset tracking of pipeline changes between execution
 	bPipelineChanged = false;
 
+	SceneOutliner->FullRefresh();
+
 	// Get stats now that the pipeline has executed
 	PostExecuteStatsPtr = FDataprepStats::GenerateWorldStats(PreviewWorld);
 	RefreshStatsTab();
diff --git a/Engine/Plugins/Enterprise/DataprepEditor/Source/DataprepEditor/Private/DataprepEditorOutlinerMode.h b/Engine/Plugins/Enterprise/DataprepEditor/Source/DataprepEditor/Private/DataprepEditorOutlinerMode.h
--- a/Engine/Plugins/Enterprise/DataprepEditor/Source/DataprepEditor/Private/DataprepEditorOutlinerMode.h
+++ b/Engine/Plugins/Enterprise/DataprepEditor/Source/DataprepEditor/Private/DataprepEditorOutlinerMode.h
@@ -52,6 +52,11 @@ public:
 		FCoreDelegates::OnActorLabelChanged.AddRaw(this, &FDataprepEditorOutlinerMode::OnActorLabelChanged);
 	}
 
+	~FDataprepEditorOutlinerMode()
+	{
+		FCoreDelegates::OnActorLabelChanged.RemoveAll(this);
+	}
+
 	virtual bool CanRenameItem(const ISceneOutlinerTreeItem& Item) const override { return false; }
 	virtual ESelectionMode::Type GetSelectionMode() const override { return ESelectionMode::Multi; }
 
