Fixing static analysis warning

diff --git a/Engine/Plugins/Enterprise/DataprepEditor/Source/DataprepCore/Private/DataprepCoreUtils.cpp b/Engine/Plugins/Enterprise/DataprepEditor/Source/DataprepCore/Private/DataprepCoreUtils.cpp
--- a/Engine/Plugins/Enterprise/DataprepEditor/Source/DataprepCore/Private/DataprepCoreUtils.cpp
+++ b/Engine/Plugins/Enterprise/DataprepEditor/Source/DataprepCore/Private/DataprepCoreUtils.cpp
@@ -216,7 +216,7 @@ void FDataprepCoreUtils::PurgeObjects(TArray<UObject*> InObjects)
 
 bool FDataprepCoreUtils::IsAsset(UObject* Object)
 {
-	const bool bHasValidObjectFlags = IsValid(Object) && !Object->HasAnyFlags(RF_ClassDefaultObject) && Object->HasAnyFlags(RF_Public);
+	const bool bHasValidObjectFlags = Object && IsValidChecked(Object) && !Object->HasAnyFlags(RF_ClassDefaultObject) && Object->HasAnyFlags(RF_Public);
 
 	if(!bHasValidObjectFlags)
 	{
