Dataprep crash when deleting dataprep asset

diff --git a/Engine/Plugins/Enterprise/DataprepEditor/Source/DataprepCore/Private/Parameterization/DataprepParameterization.cpp b/Engine/Plugins/Enterprise/DataprepEditor/Source/DataprepCore/Private/Parameterization/DataprepParameterization.cpp
--- a/Engine/Plugins/Enterprise/DataprepEditor/Source/DataprepCore/Private/Parameterization/DataprepParameterization.cpp
+++ b/Engine/Plugins/Enterprise/DataprepEditor/Source/DataprepCore/Private/Parameterization/DataprepParameterization.cpp
@@ -550,7 +550,6 @@ uint32 GetTypeHash(const TArray<FDataprepPropertyLink>& PropertyLinks)
 	return Hash;
 }
 
-
 bool UDataprepParameterizationBindings::ContainsBinding(const TSharedRef<FDataprepParameterizationBinding>& Binding) const
 {
 	return BindingToParameterName.Contains( Binding );
@@ -896,23 +895,6 @@ UDataprepParameterization::UDataprepParameterization()
 	OnObjectModifiedHandle = FCoreUObjectDelegates::OnObjectModified.AddUObject( this, &UDataprepParameterization::OnObjectModified );
 }
 
-UDataprepParameterization::~UDataprepParameterization()
-{
-	FCoreUObjectDelegates::OnObjectModified.Remove( OnObjectModifiedHandle );
-
-	if ( BindingsContainer )
-	{
-		TArray<UDataprepParameterizableObject*> ParameterizedObjects = BindingsContainer->GetParameterizedObjects();
-		for ( UDataprepParameterizableObject* Object : ParameterizedObjects )
-		{
-			if ( Object )
-			{
-				RemoveBindingToPostEditOfParameterizableObject( *Object, false );
-			}
-		}
-	}
-}
-
 void UDataprepParameterization::PostInitProperties()
 {
 	Super::PostInitProperties();
@@ -1003,6 +985,24 @@ void UDataprepParameterization::PostTransacted(const FTransactionObjectEvent& Tr
 	}
 }
 
+void UDataprepParameterization::FinishDestroy()
+{
+	FCoreUObjectDelegates::OnObjectModified.Remove( OnObjectModifiedHandle );
+
+	if ( BindingsContainer )
+	{
+		TArray<UDataprepParameterizableObject*> ParameterizedObjects = BindingsContainer->GetParameterizedObjects();
+		for ( UDataprepParameterizableObject* Object : ParameterizedObjects )
+		{
+			if ( Object )
+			{
+				RemoveBindingToPostEditOfParameterizableObject( *Object, false );
+			}
+		}
+	}
+
+	Super::FinishDestroy();
+}
 
 void UDataprepParameterization::OnObjectModified(UObject* Object)
 {
diff --git a/Engine/Plugins/Enterprise/DataprepEditor/Source/DataprepCore/Private/Parameterization/DataprepParameterization.h b/Engine/Plugins/Enterprise/DataprepEditor/Source/DataprepCore/Private/Parameterization/DataprepParameterization.h
--- a/Engine/Plugins/Enterprise/DataprepEditor/Source/DataprepCore/Private/Parameterization/DataprepParameterization.h
+++ b/Engine/Plugins/Enterprise/DataprepEditor/Source/DataprepCore/Private/Parameterization/DataprepParameterization.h
@@ -242,7 +242,6 @@ public:
 	GENERATED_BODY()
 
 	UDataprepParameterization();
-	virtual ~UDataprepParameterization();
 
 	// UObject interface
 	virtual void PostInitProperties() override;
@@ -250,6 +249,7 @@ public:
 	virtual void Serialize(FArchive& Ar) override;
 	static void AddReferencedObjects(UObject* InThis, FReferenceCollector& Collector);
 	virtual void PostTransacted(const FTransactionObjectEvent& TransactionEvent) override;
+	virtual void FinishDestroy() override;
 	// End of UObject interface
 
 	void OnObjectModified(UObject* Object);
