Dataprep: Pressing "del" on the recipe area, crashes the editor

diff --git a/Engine/Plugins/Enterprise/DataprepEditor/Source/DataprepEditor/Private/Widgets/DataprepGraph/SDataprepGraphEditor.cpp b/Engine/Plugins/Enterprise/DataprepEditor/Source/DataprepEditor/Private/Widgets/DataprepGraph/SDataprepGraphEditor.cpp
--- a/Engine/Plugins/Enterprise/DataprepEditor/Source/DataprepEditor/Private/Widgets/DataprepGraph/SDataprepGraphEditor.cpp
+++ b/Engine/Plugins/Enterprise/DataprepEditor/Source/DataprepEditor/Private/Widgets/DataprepGraph/SDataprepGraphEditor.cpp
@@ -555,6 +555,10 @@ void SDataprepGraphEditor::DeleteSelectedNodes()
 	{
 		Transaction.Cancel();
 	}
+	else
+	{
+		ClearSelectionSet();
+	}
 }
 
 bool SDataprepGraphEditor::CanDeleteNodes() const
