Dataprep: Wrong dataprep recipe items delete

diff --git a/Engine/Plugins/Enterprise/DataprepEditor/Source/DataprepCore/Private/DataprepActionAsset.cpp b/Engine/Plugins/Enterprise/DataprepEditor/Source/DataprepCore/Private/DataprepActionAsset.cpp
--- a/Engine/Plugins/Enterprise/DataprepEditor/Source/DataprepCore/Private/DataprepActionAsset.cpp
+++ b/Engine/Plugins/Enterprise/DataprepEditor/Source/DataprepCore/Private/DataprepActionAsset.cpp
@@ -510,7 +510,14 @@ bool UDataprepActionAsset::RemoveSteps(const TArray<int32>& Indices, bool bDisca
 
 	Modify();
 
-	for(int32 Index : Indices)
+	// Remove steps back to front so array indices stay valid
+	TArray<int32> SortedIndices( Indices );
+	SortedIndices.Sort( []( int32 Idx0, int32 Idx1 ) -> bool
+	{
+		return Idx0 > Idx1;
+	});
+
+	for(int32 Index : SortedIndices)
 	{
 		if ( Steps.IsValidIndex( Index ) )
 		{
