Dataprep: Black marks appear on asset texture when using Dataprep Editor

diff --git a/Engine/Plugins/Enterprise/DataprepEditor/Source/DataprepCore/Private/DataprepCorePrivateUtils.cpp b/Engine/Plugins/Enterprise/DataprepEditor/Source/DataprepCore/Private/DataprepCorePrivateUtils.cpp
--- a/Engine/Plugins/Enterprise/DataprepEditor/Source/DataprepCore/Private/DataprepCorePrivateUtils.cpp
+++ b/Engine/Plugins/Enterprise/DataprepEditor/Source/DataprepCore/Private/DataprepCorePrivateUtils.cpp
@@ -131,6 +131,7 @@ void DataprepCorePrivateUtils::BuildStaticMeshes(TSet<UStaticMesh*>& StaticMeshe
 
 					SourceModel.BuildSettings.bRecomputeNormals = !(Attributes.GetVertexInstanceNormals().IsValid() && Attributes.GetVertexInstanceNormals().GetNumChannels() > 0);
 					SourceModel.BuildSettings.bRecomputeTangents = false;
+					SourceModel.BuildSettings.DistanceFieldResolutionScale = 0.5f;
 					//SourceModel.BuildSettings.bBuildReversedIndexBuffer = false;
 				}
 			}
