Dataprep - Mac - Crash when trying to interchange Actions sequence

diff --git a/Engine/Plugins/Enterprise/DataprepEditor/Source/DataprepEditor/Private/Widgets/DataprepGraph/SDataprepGraphTrackNode.cpp b/Engine/Plugins/Enterprise/DataprepEditor/Source/DataprepEditor/Private/Widgets/DataprepGraph/SDataprepGraphTrackNode.cpp
--- a/Engine/Plugins/Enterprise/DataprepEditor/Source/DataprepEditor/Private/Widgets/DataprepGraph/SDataprepGraphTrackNode.cpp
+++ b/Engine/Plugins/Enterprise/DataprepEditor/Source/DataprepEditor/Private/Widgets/DataprepGraph/SDataprepGraphTrackNode.cpp
@@ -930,7 +930,7 @@ bool SDataprepGraphTrackNode::OnNodeDragged( TSharedPtr<SDataprepGraphBaseAction
 	const FVector2D TrackSize = TrackWidgetPtr->GetWorkingSize() * ZoomAmount;
 	float PanAmount = 0.f;
 
-	DragCallback LocalComputePanAmount = [&](const FVector2D& SlotRange, float Delta) -> float
+	DragCallback LocalComputePanAmount = [this, GraphPanel, ZoomAmount, TrackSize, PanelSize, DragTrackPosition, &TrackGeometry, &PanAmount](const FVector2D& SlotRange, float Delta) -> float
 	{
 		if(TrackSize.X > PanelSize.X)
 		{
@@ -1877,4 +1877,4 @@ void SDataprepGraphTrackWidget::UpdateLine(int32 LastIndex)
 	}
 }
 
-#undef LOCTEXT_NAMESPACE
\ No newline at end of file
+#undef LOCTEXT_NAMESPACE
