Dataprep: Filter for components does not work well in visual dataprep

diff --git a/Engine/Plugins/Enterprise/DataprepEditor/Source/DataprepEditor/Private/Widgets/DataprepAssetView.cpp b/Engine/Plugins/Enterprise/DataprepEditor/Source/DataprepEditor/Private/Widgets/DataprepAssetView.cpp
--- a/Engine/Plugins/Enterprise/DataprepEditor/Source/DataprepEditor/Private/Widgets/DataprepAssetView.cpp
+++ b/Engine/Plugins/Enterprise/DataprepEditor/Source/DataprepEditor/Private/Widgets/DataprepAssetView.cpp
@@ -206,7 +206,7 @@ void SGraphNodeDetailsWidget::OnSCSEditorTreeViewSelectionChanged(const TArray<T
 				}
 			}
 
-			UpdateFromObjects(DetailsObjects, true);
+			PropertyView->SetObjects(DetailsObjects);
 		}
 	}
 }
