Dataprep: Proxy mesh do not clean original components / actors

diff --git a/Engine/Plugins/Enterprise/DataprepEditor/Source/DataprepLibraries/Private/DataprepEditingOperations.cpp b/Engine/Plugins/Enterprise/DataprepEditor/Source/DataprepLibraries/Private/DataprepEditingOperations.cpp
--- a/Engine/Plugins/Enterprise/DataprepEditor/Source/DataprepLibraries/Private/DataprepEditingOperations.cpp
+++ b/Engine/Plugins/Enterprise/DataprepEditor/Source/DataprepLibraries/Private/DataprepEditingOperations.cpp
@@ -81,6 +81,7 @@ namespace DatasmithEditingOperationsUtils
 	void FindActorsToCollapseOrDelete(const TArray<AActor*>& ActorsToVisit, TArray<AActor*>& ActorsToCollapse, TArray<UObject*>& ActorsToDelete );
 	void GetRootActors(UWorld* World, TArray<AActor*>& RootActors);
 	void GetComponentsToMerge(UWorld*& World, const TArray<UObject*>& InObjects, TArray<UStaticMeshComponent*>& ComponentsToMerge);
+	TArray<UObject*> CollectObjectsToDeleteAfterMeshMerging(TArray<UStaticMeshComponent*>& InMergedComponents);
 
 	int32 GetActorDepth(AActor* Actor)
 	{
@@ -209,43 +210,12 @@ void UDataprepMergeActorsOperation::OnExecution_Implementation(const FDataprepCo
 	MergedActor->GetRootComponent()->AttachToComponent( ComponentsToMerge[0]->GetAttachParent(), FAttachmentTransformRules::KeepWorldTransform );
 
 	// Collect all objects to be deleted
+	TArray<UObject*> ObjectsToDelete = DatasmithEditingOperationsUtils::CollectObjectsToDeleteAfterMeshMerging(ComponentsToMerge);
 
-	// Sort merged components to detach children first
-	ComponentsToMerge.Sort([](const USceneComponent &A, const USceneComponent &B) -> bool { return A.IsAttachedTo(&B); });
-
-	TSet<AActor*> ActorsToDelete;
-
-	// Simple way to delete the actors: detach the merged components if it's safe to do so
-	for(UStaticMeshComponent* Component : ComponentsToMerge)
+	if( ObjectsToDelete.Num() > 0 )
 	{
-		ActorsToDelete.Add(Component->GetOwner());
-
-		if(Component->GetNumChildrenComponents() == 0)
-		{
-			Component->DetachFromComponent(FDetachmentTransformRules::KeepRelativeTransform);
-		}
-		else
-		{
-			// To keep the children in the hierarchy, simply hide the StaticMesh, but the StaticMeshComponent could be replaced by a SceneComponent and
-			// then attach the children to it
-			Component->SetStaticMesh(nullptr);
-		}
+		DeleteObjects( ObjectsToDelete );
 	}
-
-	TArray<UObject*> ObjectsToDelete;
-	ObjectsToDelete.Reserve(ActorsToDelete.Num());
-
-	// Then delete the merged actors that don't have any children component
-	for(AActor* Actor : ActorsToDelete)
-	{
-		USceneComponent* RootComponent = Cast<USceneComponent>(Actor->GetRootComponent());
-		if(RootComponent && RootComponent->GetNumChildrenComponents() == 0)
-		{
-			ObjectsToDelete.Add(Actor);
-		}
-	}
-
-	DeleteObjects(ObjectsToDelete);
 }
 
 bool UDataprepMergeActorsOperation::MergeStaticMeshActors(UWorld* World, const TArray<UStaticMeshComponent*>& ComponentsToMerge, const FString& RootName, bool bCreateActor)
@@ -469,43 +439,12 @@ void UDataprepCreateProxyMeshOperation::OnExecution_Implementation(const FDatapr
 	MergedActor->GetRootComponent()->AttachToComponent(ComponentsToMerge[0]->GetAttachParent(), FAttachmentTransformRules::KeepWorldTransform);
 
 	// Collect all objects to be deleted
+	TArray<UObject*> ObjectsToDelete = DatasmithEditingOperationsUtils::CollectObjectsToDeleteAfterMeshMerging(ComponentsToMerge);
 
-	// Sort merged components to detach children first
-	ComponentsToMerge.Sort([](const USceneComponent &A, const USceneComponent &B) -> bool { return A.IsAttachedTo(&B); });
-
-	TSet<AActor*> ActorsToDelete;
-
-	// Simple way to delete the actors: detach the merged components if it's safe to do so
-	for(UStaticMeshComponent* Component : ComponentsToMerge)
-	{
-		ActorsToDelete.Add(Component->GetOwner());
-
-		if(Component->GetNumChildrenComponents() == 0)
-		{
-			Component->DetachFromComponent(FDetachmentTransformRules::KeepRelativeTransform);
-		}
-		else
-		{
-			// To keep the children in the hierarchy, simply hide the StaticMesh, but the StaticMeshComponent could be replaced by a SceneComponent and
-			// then attach the children to it
-			Component->SetStaticMesh(nullptr);
-		}
-	}
-
-	TArray<UObject*> ObjectsToDelete;
-	ObjectsToDelete.Reserve(ActorsToDelete.Num());
-
-	// Then delete the merged actors that don't have any children component
-	for(AActor* Actor : ActorsToDelete)
+	if( ObjectsToDelete.Num() > 0 )
 	{
-		USceneComponent* RootComponent = Cast<USceneComponent>(Actor->GetRootComponent());
-		if(RootComponent && RootComponent->GetNumChildrenComponents() == 0)
-		{
-			ObjectsToDelete.Add(Actor);
-		}
+		DeleteObjects( ObjectsToDelete );
 	}
-
-	DeleteObjects(ObjectsToDelete);
 }
 
 void UDataprepDeleteUnusedAssetsOperation::OnExecution_Implementation(const FDataprepContext& InContext)
@@ -1101,6 +1040,66 @@ namespace DatasmithEditingOperationsUtils
 		}
 	}
 
+	TArray<UObject*> CollectObjectsToDeleteAfterMeshMerging(TArray<UStaticMeshComponent*>& InMergedComponents)
+	{
+		// Sort merged components to detach children first
+		InMergedComponents.Sort([](const USceneComponent &A, const USceneComponent &B) -> bool { return A.IsAttachedTo(&B); });
+
+		TArray<UObject*> ObjectsToDelete;
+		ObjectsToDelete.Reserve(InMergedComponents.Num());
+
+		for( UStaticMeshComponent* Component : InMergedComponents )
+		{
+			if( Component->GetNumChildrenComponents() > 0 )
+			{
+				// Reparent children
+				USceneComponent* NewParentForChildren = nullptr;
+
+				if( Component->GetOwner()->GetRootComponent() == Component )
+				{
+					// Promote first child to become new root
+					USceneComponent* FirstChild = Component->GetChildComponent( 0 );
+					Component->GetOwner()->SetRootComponent( FirstChild );
+					NewParentForChildren = FirstChild;
+				}
+
+				if( USceneComponent* Parent = Component->GetAttachParent() )
+				{
+					Component->GetChildComponent( 0 )->AttachToComponent( Parent, FAttachmentTransformRules::KeepWorldTransform );
+
+					if( !NewParentForChildren )
+					{
+						NewParentForChildren = Parent;
+					}
+				}
+
+				for( int32 ChildIndex = 1; ChildIndex < Component->GetNumChildrenComponents(); ++ChildIndex )
+				{
+					Component->GetChildComponent( ChildIndex )->AttachToComponent( NewParentForChildren, FAttachmentTransformRules::KeepWorldTransform );
+				}
+			}
+
+			Component->DetachFromComponent(FDetachmentTransformRules::KeepRelativeTransform);
+
+			// Check if we can delete actor
+			USceneComponent* RootComponent = Cast<USceneComponent>(Component->GetOwner()->GetRootComponent());
+			check( RootComponent );
+
+			const bool bCanRemoveActor = ( RootComponent == Component ) || ( RootComponent->GetNumChildrenComponents() == 0 && RootComponent->GetClass() == USceneComponent::StaticClass() );
+
+			if( bCanRemoveActor )
+			{
+				ObjectsToDelete.Add( Component->GetOwner() );
+			}
+			else
+			{
+				ObjectsToDelete.Add( Component );
+			}
+		}
+
+		return MoveTemp( ObjectsToDelete );
+	}
+
 	FMergingData::FMergingData(const TArray<UPrimitiveComponent*>& PrimitiveComponents)
 	{
 		Data.Reserve(PrimitiveComponents.Num());
