Merged from UE4/Main (CL #16445731), robo-merge fix: Dataprep: Update filter from Selection : simpler UI and use object / actors path as the unique ID

diff --git a/Engine/Plugins/Enterprise/DataprepEditor/Source/DataprepCore/Private/DataprepCorePrivateUtils.cpp b/Engine/Plugins/Enterprise/DataprepEditor/Source/DataprepCore/Private/DataprepCorePrivateUtils.cpp
--- a/Engine/Plugins/Enterprise/DataprepEditor/Source/DataprepCore/Private/DataprepCorePrivateUtils.cpp
+++ b/Engine/Plugins/Enterprise/DataprepEditor/Source/DataprepCore/Private/DataprepCorePrivateUtils.cpp
@@ -131,7 +131,7 @@ void DataprepCorePrivateUtils::BuildStaticMeshes(TSet<UStaticMesh*>& StaticMeshe
 
 					SourceModel.BuildSettings.bRecomputeNormals = !(Attributes.GetVertexInstanceNormals().IsValid() && Attributes.GetVertexInstanceNormals().GetNumChannels() > 0);
 					SourceModel.BuildSettings.bRecomputeTangents = false;
-					SourceModel.BuildSettings.DistanceFieldResolutionScale = 0.5f;
+					SourceModel.BuildSettings.DistanceFieldResolutionScale = 0;
 					//SourceModel.BuildSettings.bBuildReversedIndexBuffer = false;
 				}
 			}
diff --git a/Engine/Plugins/Enterprise/DataprepEditor/Source/DataprepCore/Private/SelectionSystem/DataprepObjectSelectionFilter.cpp b/Engine/Plugins/Enterprise/DataprepEditor/Source/DataprepCore/Private/SelectionSystem/DataprepObjectSelectionFilter.cpp
new file mode 100644
--- /dev/null
+++ b/Engine/Plugins/Enterprise/DataprepEditor/Source/DataprepCore/Private/SelectionSystem/DataprepObjectSelectionFilter.cpp
@@ -0,0 +1,131 @@
+// Copyright Epic Games, Inc. All Rights Reserved.
+
+#include "SelectionSystem/DataprepObjectSelectionFilter.h"
+#include "SelectionSystem/DataprepSelectionSystemStructs.h"
+#include "Shared/DataprepCorePrivateUtils.h"
+
+#include "Engine/Texture.h"
+#include "Engine/StaticMesh.h"
+#include "Materials/MaterialInterface.h"
+#include "GameFramework/Actor.h"
+
+TArray<UObject*> UDataprepObjectSelectionFilter::FilterObjects(const TArrayView<UObject*>& Objects) const
+{
+	TArray<UObject*> FilteredObjects;
+
+	RunFilter(Objects, FilteredObjects, nullptr);
+
+	return FilteredObjects;
+}
+
+void UDataprepObjectSelectionFilter::FilterAndGatherInfo(const TArrayView<UObject*>& InObjects, const TArrayView<FDataprepSelectionInfo>& OutFilterResults) const
+{
+	TArray<UObject*> FilteredObjects;
+	TArray<bool> FilterResults;
+	FilterResults.Init(false, InObjects.Num());
+	TArrayView<bool> FilterResultsView(FilterResults);
+
+	RunFilter(InObjects, FilteredObjects, &FilterResultsView);
+
+	for (int Index = 0; Index < InObjects.Num(); ++Index)
+	{
+		FDataprepSelectionInfo& SelectionInfo = OutFilterResults[Index];
+		SelectionInfo.bHasPassFilter = FilterResultsView[Index];
+		SelectionInfo.bWasDataFetchedAndCached = false;
+	}
+}
+
+void UDataprepObjectSelectionFilter::FilterAndStoreInArrayView(const TArrayView<UObject*>& InObjects, const TArrayView<bool>& OutFilterResults) const
+{
+	TArray<UObject*> FilteredObjects;
+	RunFilter(InObjects, FilteredObjects, &OutFilterResults);
+}
+
+void UDataprepObjectSelectionFilter::SetSelection( const FString& InTransientContentPath, const TArray<UObject*>& InSelectedObjects )
+{
+	NumAssets = 0;
+	NumActors = 0;
+
+	TFunction<bool( UObject* )> IsAsset = []( const UObject* InObject ) -> bool
+	{
+		return
+			Cast<UMaterialInterface>( InObject ) ||
+			Cast<UStaticMesh>( InObject ) ||
+			Cast<UTexture>( InObject );
+	};
+
+	for( UObject* Obj : InSelectedObjects )
+	{
+		FString ObjectName;
+
+		if( IsAsset( Obj ) )
+		{
+			++NumAssets;
+
+			// Cutoff transient content folder from full path name
+			FString PartialPath = Obj->GetPathName();
+			PartialPath.RemoveFromStart( InTransientContentPath );
+
+			SelectedObjectPaths.Add( PartialPath );
+
+			ObjectName = Obj->GetName();
+
+		}
+		else if( AActor* Actor = Cast<AActor>( Obj ) )
+		{
+			// Get path inside level only
+			for( UObject* Outer = Obj->GetOuter(); Outer; Outer = Outer->GetOuter() )
+			{
+				if( ULevel* LevelOuter = Cast<ULevel>( Outer ) )
+				{
+					++NumActors;
+					SelectedObjectPaths.Add( Obj->GetPathName( Outer ) );
+
+					ObjectName = Actor->GetActorLabel();
+
+					break;
+				}
+			}
+		}
+
+		if( ObjectName.Len() > 0 )
+		{
+			CachedNames.Add( ObjectName );
+		}
+	}
+}
+
+void UDataprepObjectSelectionFilter::RunFilter(const TArrayView<UObject*>& InputObjects, TArray<UObject*>& FilteredObjects, const TArrayView<bool>* OutFilterResults) const
+{
+	TArray<AActor*> InputActors;
+
+	FilteredObjects.Empty();
+
+	for (UObject* Object : InputObjects)
+	{
+		for( const FString& SelPath : SelectedObjectPaths )
+		{
+			if( Object->GetPathName().EndsWith( SelPath ) )
+			{
+				FilteredObjects.Add( Object );
+				break;
+			}
+		}
+	}
+
+	if (OutFilterResults != nullptr)
+	{
+		check(OutFilterResults->Num() >= InputObjects.Num());
+
+		for (int ObjectIndex = 0; ObjectIndex < InputObjects.Num(); ++ObjectIndex)
+		{
+			bool& bFiltered = (*OutFilterResults)[ObjectIndex];
+			bFiltered = (INDEX_NONE != FilteredObjects.Find(InputObjects[ObjectIndex]));
+
+			if (IsExcludingResult())
+			{
+				bFiltered = !bFiltered;
+			}
+		}
+	}
+}
diff --git a/Engine/Plugins/Enterprise/DataprepEditor/Source/DataprepCore/Public/SelectionSystem/DataprepObjectSelectionFilter.h b/Engine/Plugins/Enterprise/DataprepEditor/Source/DataprepCore/Public/SelectionSystem/DataprepObjectSelectionFilter.h
new file mode 100644
--- /dev/null
+++ b/Engine/Plugins/Enterprise/DataprepEditor/Source/DataprepCore/Public/SelectionSystem/DataprepObjectSelectionFilter.h
@@ -0,0 +1,54 @@
+// Copyright Epic Games, Inc. All Rights Reserved.
+
+#pragma once
+
+#include "SelectionSystem/DataprepFilter.h"
+
+#include "CoreMinimal.h"
+#include "UObject/Object.h"
+#include "UObject/ObjectMacros.h"
+
+#include "DataprepObjectSelectionFilter.generated.h"
+
+UCLASS(BlueprintType, NotBlueprintable, HideCategories = (Filter), Meta = (Hidden, DisplayName="Selection", ToolTip = "Filter selected objects"))
+class DATAPREPCORE_API UDataprepObjectSelectionFilter : public UDataprepFilterNoFetcher
+{
+	GENERATED_BODY()
+
+public:
+	UDataprepObjectSelectionFilter() : NumAssets( 0 ), NumActors( 0 ) {}
+
+	//~ Begin UDataprepFilterNoFetcher Interface
+	virtual TArray<UObject*> FilterObjects(const TArrayView<UObject*>& Objects) const override;
+	virtual void FilterAndGatherInfo(const TArrayView<UObject*>& InObjects, const TArrayView<FDataprepSelectionInfo>& OutFilterResults) const override;
+	virtual void FilterAndStoreInArrayView(const TArrayView<UObject*>& InObjects, const TArrayView<bool>& OutFilterResults) const override;
+	virtual bool IsThreadSafe() const override { return true; }
+	//~ End UDataprepFilterNoFetcher Interface
+
+	void SetSelection( const FString& InTransientContentPath, const TArray<UObject*>& InSelectedObjects );
+
+	const TArray< FString >& GetCachedNames() const { return CachedNames; }
+
+	int32 GetNumAssets() const { return NumAssets; }
+	int32 GetNumActors() const { return NumActors; }
+
+private:
+	void RunFilter(const TArrayView<UObject*>& InputObjects, TArray<UObject*>& FilteredObjects, const TArrayView<bool>* OutFilterResults) const;
+
+private:
+	// Partial paths of objects.
+	// For assets, the transient part is cutoff (as it changes between runs), 
+	// for actors, the path up to level outer is cutoff (in case level name changes between runs)
+	UPROPERTY()
+	TArray< FString > SelectedObjectPaths;
+
+	UPROPERTY()
+	int32 NumAssets;
+
+	UPROPERTY()
+	int32 NumActors;
+
+	// Cache some of the names of selected objects, up to some limit, for preview purpose
+	UPROPERTY()
+	TArray< FString > CachedNames;
+};
diff --git a/Engine/Plugins/Enterprise/DataprepEditor/Source/DataprepEditor/Private/PreviewSystem/DataprepPreviewSystem.cpp b/Engine/Plugins/Enterprise/DataprepEditor/Source/DataprepEditor/Private/PreviewSystem/DataprepPreviewSystem.cpp
--- a/Engine/Plugins/Enterprise/DataprepEditor/Source/DataprepEditor/Private/PreviewSystem/DataprepPreviewSystem.cpp
+++ b/Engine/Plugins/Enterprise/DataprepEditor/Source/DataprepEditor/Private/PreviewSystem/DataprepPreviewSystem.cpp
@@ -225,7 +225,6 @@ void FDataprepPreviewSystem::SetObservedObjects(const TArrayView<UDataprepParame
 			FDelegateHandle Handle = Object->GetOnPostEdit().AddSP( this, &FDataprepPreviewSystem::OnObservedObjectPostEdit );
 			ObservedOnPostEdit.Add( Object, Handle );
 
-
 			if ( UDataprepFilter* Filter  = Cast<UDataprepFilter>( Object ) )
 			{
 				if ( UDataprepFetcher* Fetcher = Filter->GetFetcher() )
@@ -234,6 +233,11 @@ void FDataprepPreviewSystem::SetObservedObjects(const TArrayView<UDataprepParame
 					ObservedOnPostEdit.Add( Fetcher, Handle );
 				}
 			}
+			else if( UDataprepFilterNoFetcher* FilterNF = Cast<UDataprepFilterNoFetcher>( Object ) )
+			{
+				Handle = FilterNF->GetOnPostEdit().AddSP( this, &FDataprepPreviewSystem::OnObservedObjectPostEdit);
+				ObservedOnPostEdit.Add( FilterNF, Handle );
+			}
 		}
 	}
 
@@ -287,6 +291,17 @@ void FDataprepPreviewSystem::IncrementalProcess()
 				Filter->FilterAndGatherInfo( Objects, FilterResults );
 				PopulateResultFromFilter( ItemsProcessedCount );
 			}
+			if ( UDataprepFilterNoFetcher* FilterNF = Cast<UDataprepFilterNoFetcher>( ObservedSteps[CurrentProgress.CurrentFilterIndex] ) )
+			{
+				bDidSomeProcessing = true;
+				PrepareFilterBuffers( IncrementalCount );
+				ItemsProcessedCount = FillObjectsBuffer( ItemsProcessedCount );
+
+				TArrayView<UObject*> Objects( ObjectsBuffer.GetData(), ItemsProcessedCount );
+				TArrayView<FDataprepSelectionInfo> FilterResults( FilterResultsBuffer.GetData(), ItemsProcessedCount );
+				FilterNF->FilterAndGatherInfo( Objects, FilterResults );
+				PopulateResultFromFilter( ItemsProcessedCount );
+			}
 
 			if ( CurrentProgress.CurrentObjectProcessed >= ItemsCount )
 			{
diff --git a/Engine/Plugins/Enterprise/DataprepEditor/Source/DataprepEditor/Private/SchemaActions/DataprepFilterMenuActionCollector.cpp b/Engine/Plugins/Enterprise/DataprepEditor/Source/DataprepEditor/Private/SchemaActions/DataprepFilterMenuActionCollector.cpp
--- a/Engine/Plugins/Enterprise/DataprepEditor/Source/DataprepEditor/Private/SchemaActions/DataprepFilterMenuActionCollector.cpp
+++ b/Engine/Plugins/Enterprise/DataprepEditor/Source/DataprepEditor/Private/SchemaActions/DataprepFilterMenuActionCollector.cpp
@@ -48,7 +48,7 @@ TArray<TSharedPtr<FDataprepSchemaAction>> FDataprepFilterMenuActionCollector::Co
 
 	for ( UClass* FilterClass : FilterNoFetcherClasses )
 	{
-		if ( FilterClass )
+		if ( FilterClass && !FilterClass->HasMetaData( TEXT( "Hidden" ) ) )
 		{
 			UDataprepFilterNoFetcher* Filter = FilterClass->GetDefaultObject< UDataprepFilterNoFetcher >();
 			if ( Filter )
@@ -62,8 +62,6 @@ TArray<TSharedPtr<FDataprepSchemaAction>> FDataprepFilterMenuActionCollector::Co
 		}
 	}
 
-
-
 	UE_LOG( LogDataprepEditor, Log, TEXT("The discovery of the filters/fetchers and the creation of the menu actions took %f seconds."), ( FPlatformTime::Seconds() - Start ) );
 
 	return Actions;
diff --git a/Engine/Plugins/Enterprise/DataprepEditor/Source/DataprepEditor/Private/Widgets/DataprepGraph/SDataprepFilter.cpp b/Engine/Plugins/Enterprise/DataprepEditor/Source/DataprepEditor/Private/Widgets/DataprepGraph/SDataprepFilter.cpp
--- a/Engine/Plugins/Enterprise/DataprepEditor/Source/DataprepEditor/Private/Widgets/DataprepGraph/SDataprepFilter.cpp
+++ b/Engine/Plugins/Enterprise/DataprepEditor/Source/DataprepEditor/Private/Widgets/DataprepGraph/SDataprepFilter.cpp
@@ -10,10 +10,12 @@
 #include "SelectionSystem/DataprepIntegerFilter.h"
 #include "SelectionSystem/DataprepStringFilter.h"
 #include "SelectionSystem/DataprepStringsArrayFilter.h"
+#include "SelectionSystem/DataprepObjectSelectionFilter.h"
 #include "Widgets/DataprepGraph/SDataprepBoolFilter.h"
 #include "Widgets/DataprepGraph/SDataprepFloatFilter.h"
 #include "Widgets/DataprepGraph/SDataprepIntegerFilter.h"
 #include "Widgets/DataprepGraph/SDataprepStringFilter.h"
+#include "Widgets/DataprepGraph/SDataprepObjectSelectionFilter.h"
 #include "Widgets/DataprepWidgets.h"
 
 #include "EditorStyleSet.h"
@@ -283,7 +285,23 @@ TSharedRef<SWidget> SDataprepFilterNoFetcher::GetTitleWidget()
 
 TSharedRef<SWidget> SDataprepFilterNoFetcher::GetContentWidget()
 {
+	TSharedPtr< SWidget > FilterWidget = SNullWidget::NullWidget;
+
+	if ( Filter )
+	{
+		UClass* Class = Filter->GetClass();
+		if ( Class == UDataprepObjectSelectionFilter::StaticClass() )
+		{
+			SAssignNew(FilterWidget, SDataprepObjectSelectionFilter, *static_cast< UDataprepObjectSelectionFilter* >( Filter ) );
+		}
+	}
+
 	return SNew( SVerticalBox )
+		+ SVerticalBox::Slot()
+		.AutoHeight()
+		[
+			FilterWidget.ToSharedRef()
+		]
 		+ SVerticalBox::Slot()
 		.AutoHeight()
 		[
diff --git a/Engine/Plugins/Enterprise/DataprepEditor/Source/DataprepEditor/Private/Widgets/DataprepGraph/SDataprepGraphEditor.cpp b/Engine/Plugins/Enterprise/DataprepEditor/Source/DataprepEditor/Private/Widgets/DataprepGraph/SDataprepGraphEditor.cpp
--- a/Engine/Plugins/Enterprise/DataprepEditor/Source/DataprepEditor/Private/Widgets/DataprepGraph/SDataprepGraphEditor.cpp
+++ b/Engine/Plugins/Enterprise/DataprepEditor/Source/DataprepEditor/Private/Widgets/DataprepGraph/SDataprepGraphEditor.cpp
@@ -14,6 +14,7 @@
 #include "SchemaActions/DataprepDragDropOp.h"
 #include "SchemaActions/IDataprepMenuActionCollector.h"
 #include "SelectionSystem/DataprepStringFilter.h"
+#include "SelectionSystem/DataprepObjectSelectionFilter.h"
 #include "Widgets/DataprepGraph/SDataprepGraphActionNode.h"
 #include "Widgets/DataprepGraph/SDataprepGraphActionStepNode.h"
 #include "Widgets/DataprepGraph/SDataprepGraphTrackNode.h"
@@ -809,35 +810,12 @@ void SDataprepGraphEditor::OnCollectCustomActions(TArray<TSharedPtr<FDataprepSch
 		return;
 	}
 
-	UClass* NameFetcherClass = FindObject<UClass>( ANY_PACKAGE, TEXT( "DataprepStringObjectNameFetcher" ) );
-	check( NameFetcherClass );
-
 	FDataprepSchemaAction::FOnExecuteAction OnExcuteMenuAction;
-	OnExcuteMenuAction.BindLambda( [this, AssetAndActorSelection, NameFetcherClass] (const FDataprepSchemaActionContext& InContext)
+	OnExcuteMenuAction.BindLambda( [this, AssetAndActorSelection] (const FDataprepSchemaActionContext& InContext)
 	{
-		UDataprepActionAsset* Action = InContext.DataprepActionPtr.Get();
-		if ( Action )
+		if( UDataprepActionAsset* Action = InContext.DataprepActionPtr.Get() )
 		{
-			int32 NewFilterIndex = Action->AddStep( NameFetcherClass );
-
-			UDataprepActionStep* Step = Action->GetStep(NewFilterIndex).Get();
-			check( Step );
-
-			UDataprepStringFilter* StringFilter = Cast< UDataprepStringFilter >( Step->GetStepObject() );
-			check( StringFilter );
-
-			StringFilter->SetMatchInArray( true );
-			UDataprepStringFilterMatchingArray* ArrayField = StringFilter->GetStringArray();
-
-			for ( UObject* Obj : AssetAndActorSelection )
-			{
-				ArrayField->Strings.Add( Obj->GetName() );
-			}
-
-			if(SDataprepGraphTrackNode* TrackGraphNode = TrackGraphNodePtr.Pin().Get())
-			{
-				TrackGraphNode->UpdateGraphNode();
-			}
+			CreateFilterFromSelection( Action, AssetAndActorSelection );
 		}
 	});
 
@@ -848,6 +826,30 @@ void SDataprepGraphEditor::OnCollectCustomActions(TArray<TSharedPtr<FDataprepSch
 	OutActions.Add( Action );
 }
 
+void SDataprepGraphEditor::CreateFilterFromSelection(UDataprepActionAsset* InTargetAction, const TSet< UObject* >& InAssetAndActorSelection)
+{
+	check( InTargetAction );
+
+	if ( InAssetAndActorSelection.Num() == 0 )
+	{
+		return;
+	}
+
+	TSharedPtr<FDataprepEditor> DataprepEditorPtr = DataprepEditor.Pin();
+	check( DataprepEditorPtr );
+
+	UDataprepObjectSelectionFilter* Filter = NewObject< UDataprepObjectSelectionFilter >( GetTransientPackage(), UDataprepObjectSelectionFilter::StaticClass(), NAME_None, RF_Transactional );
+	Filter->SetSelection( DataprepEditorPtr->GetTransientContentFolder(), InAssetAndActorSelection.Array() );
+
+	int32 NewFilterIndex = InTargetAction->AddStep( Filter );
+	ensure( NewFilterIndex != INDEX_NONE );
+
+	if(SDataprepGraphTrackNode* TrackGraphNode = TrackGraphNodePtr.Pin().Get())
+	{
+		TrackGraphNode->UpdateGraphNode();
+	}
+}
+
 FActionMenuContent SDataprepGraphEditor::OnCreateActionMenu(UEdGraph* InGraph, const FVector2D& InNodePosition, const TArray<UEdGraphPin*>& InDraggedPins, bool bAutoExpand, SGraphEditor::FActionMenuClosed InOnMenuClosed)
 {
 	// bAutoExpand is voluntary ignored for now
@@ -1157,28 +1159,7 @@ FActionMenuContent SDataprepGraphEditor::OnCreateNodeOrPinMenu(UEdGraph* Current
 					{
 						if ( UDataprepActionAsset* Action = Actions[0] )
 						{
-							UClass* NameFetcherClass = FindObject<UClass>( ANY_PACKAGE, TEXT( "DataprepStringObjectNameFetcher" ) );
-							check( NameFetcherClass );
-							const int32 StepIndex = Action->AddStep( NameFetcherClass );
-
-							UDataprepActionStep* Step = Action->GetStep(StepIndex).Get();
-							check( Step );
-
-							UDataprepStringFilter* StringFilter = Cast< UDataprepStringFilter >( Step->GetStepObject() );
-							check( StringFilter );
-
-							StringFilter->SetMatchInArray( true );
-							UDataprepStringFilterMatchingArray* ArrayField = StringFilter->GetStringArray();
-
-							for ( UObject* Obj : AssetAndActorSelection )
-							{
-								ArrayField->Strings.Add( Obj->GetName() );
-							}
-
-							if(SDataprepGraphTrackNode* TrackGraphNode = TrackGraphNodePtr.Pin().Get())
-							{
-								TrackGraphNode->UpdateGraphNode();
-							}
+							CreateFilterFromSelection( Action, AssetAndActorSelection );
 						}
 					});
 					
@@ -1197,7 +1178,7 @@ FActionMenuContent SDataprepGraphEditor::OnCreateNodeOrPinMenu(UEdGraph* Current
 			bool bIsSelectionOnlyFilters = true;
 			bool bAreFilterFromSameAction = true;
 			
-			TArray<UDataprepFilter*> Filters;
+			TArray<UDataprepParameterizableObject*> Filters;
 			Filters.Reserve( SelectedNodes.Num() );
 
 			const UDataprepActionAsset* ClickedAction = FirstStepNode->GetDataprepActionAsset();
@@ -1226,6 +1207,10 @@ FActionMenuContent SDataprepGraphEditor::OnCreateNodeOrPinMenu(UEdGraph* Current
 					{
 						Filters.Add( Filter );
 					}
+					else if( UDataprepFilterNoFetcher* FilterNF = Cast<UDataprepFilterNoFetcher>( StepObject ) )
+					{
+						Filters.Add( FilterNF );
+					}
 					else
 					{
 						bIsSelectionOnlyFilters = false;
@@ -1243,9 +1228,16 @@ FActionMenuContent SDataprepGraphEditor::OnCreateNodeOrPinMenu(UEdGraph* Current
 					InverseFilterAction.ExecuteAction.BindLambda( [Filters]()
 						{
 							FScopedTransaction Transaction( LOCTEXT("InverseFilterTransaction", "Inverse the filter") );
-							for ( UDataprepFilter* Filter : Filters )
+							for ( UDataprepParameterizableObject* ParamObj : Filters )
 							{
-								Filter->SetIsExcludingResult( !Filter->IsExcludingResult() );
+								if( UDataprepFilter* Filter = Cast<UDataprepFilter>( ParamObj ) )
+								{
+									Filter->SetIsExcludingResult( !Filter->IsExcludingResult() );
+								}
+								else if( UDataprepFilterNoFetcher* FilterNF = Cast<UDataprepFilterNoFetcher>( ParamObj ) )
+								{
+									FilterNF->SetIsExcludingResult( !FilterNF->IsExcludingResult() );
+								}
 							}
 						});
 					MenuBuilder->AddMenuEntry(LOCTEXT("InverseFilter", "Inverse Filter(s) Selection"),
@@ -1263,7 +1255,7 @@ FActionMenuContent SDataprepGraphEditor::OnCreateNodeOrPinMenu(UEdGraph* Current
 						 * Work with the assumption that the filters array contains only unique objects
 						 */
 						{
-							for ( UDataprepFilter* Filter : Filters )
+							for ( UDataprepParameterizableObject* Filter : Filters )
 							{
 								if ( !DataprepEditorPtr->IsPreviewingStep(Filter) )
 								{
diff --git a/Engine/Plugins/Enterprise/DataprepEditor/Source/DataprepEditor/Private/Widgets/DataprepGraph/SDataprepGraphEditor.h b/Engine/Plugins/Enterprise/DataprepEditor/Source/DataprepEditor/Private/Widgets/DataprepGraph/SDataprepGraphEditor.h
--- a/Engine/Plugins/Enterprise/DataprepEditor/Source/DataprepEditor/Private/Widgets/DataprepGraph/SDataprepGraphEditor.h
+++ b/Engine/Plugins/Enterprise/DataprepEditor/Source/DataprepEditor/Private/Widgets/DataprepGraph/SDataprepGraphEditor.h
@@ -136,6 +136,8 @@ public:
 	/** Callback from the contextual menu to create any entries outside of the action collector */
 	void OnCollectCustomActions(TArray<TSharedPtr<FDataprepSchemaAction>>& OutActions);
 
+	void CreateFilterFromSelection(UDataprepActionAsset* InTargetAction, const TSet< UObject* >& InAssetAndActorSelection);
+
 private:
 	/** Recompute the layout of the displayed graph after a pan, resize and/or zoom */
 	void UpdateLayout( const FVector2D& LocalSize, const FVector2D& Location, float ZoomAmount );
diff --git a/Engine/Plugins/Enterprise/DataprepEditor/Source/DataprepEditor/Private/Widgets/DataprepGraph/SDataprepObjectSelectionFilter.cpp b/Engine/Plugins/Enterprise/DataprepEditor/Source/DataprepEditor/Private/Widgets/DataprepGraph/SDataprepObjectSelectionFilter.cpp
new file mode 100644
--- /dev/null
+++ b/Engine/Plugins/Enterprise/DataprepEditor/Source/DataprepEditor/Private/Widgets/DataprepGraph/SDataprepObjectSelectionFilter.cpp
@@ -0,0 +1,84 @@
+// Copyright Epic Games, Inc. All Rights Reserved.
+
+#include "Widgets/DataprepGraph/SDataprepObjectSelectionFilter.h"
+#include "SelectionSystem/DataprepObjectSelectionFilter.h"
+
+#include "Widgets/Text/STextBlock.h"
+#include "Widgets/SBoxPanel.h"
+
+#define LOCTEXT_NAMESPACE "DataprepObjectSelectionFilter"
+
+void SDataprepObjectSelectionFilter::Construct(const FArguments& InArgs, UDataprepObjectSelectionFilter& InFilter)
+{
+	Filter = &InFilter;
+
+	FString ObjectNames;
+
+	const int32 MaxObjectNamesLen = 250;
+	int32 ObjectNamesUsed = 0;
+
+	for( const FString& Name : Filter->GetCachedNames() )
+	{
+		if( Name.IsEmpty() )
+		{
+			continue;
+		}
+
+		if( ObjectNames.Len() > 0 )
+		{
+			ObjectNames.Append( ", " );
+		}
+
+		ObjectNames.Append( Name );
+
+		++ObjectNamesUsed;
+
+		if( ObjectNames.Len() >= MaxObjectNamesLen )
+		{
+			break;
+		}
+	}
+
+	if( ObjectNamesUsed < Filter->GetCachedNames().Num() )
+	{
+		ObjectNames.Append( "..." );
+	}
+	
+	ChildSlot
+	[
+		SNew( SVerticalBox )
+		+ SVerticalBox::Slot()
+		.Padding( 3, 2, 3, 0)
+		.AutoHeight()
+		[
+			SNew( STextBlock )
+			.ColorAndOpacity( FLinearColor::Yellow )
+			.Text( FText::Format( LOCTEXT("NumActors", "Num Actors {0}"), FText::AsNumber( Filter->GetNumActors() ) ) )
+		]
+		+ SVerticalBox::Slot()
+		.Padding( 3, 2, 3, 2)
+		.AutoHeight()
+		[
+			SNew( STextBlock )
+			.ColorAndOpacity( FLinearColor::Yellow )
+			.Text( FText::Format( LOCTEXT("NumAssets", "Num Assets {0}"), FText::AsNumber( Filter->GetNumAssets() ) ) )
+		]
+		+ SVerticalBox::Slot()
+		.Padding( 3, 5, 3, 0)
+		.AutoHeight()
+		[
+			SNew( STextBlock )
+			.AutoWrapText( true )
+			.WrappingPolicy( ETextWrappingPolicy::AllowPerCharacterWrapping )
+			.Font( FCoreStyle::GetDefaultFontStyle( "Italic", 10 ) )
+			.Text( FText::FromString( ObjectNames ) )
+		]
+	];
+}
+
+void SDataprepObjectSelectionFilter::AddReferencedObjects(FReferenceCollector& Collector)
+{
+	Collector.AddReferencedObject( Filter );
+}
+
+#undef LOCTEXT_NAMESPACE
diff --git a/Engine/Plugins/Enterprise/DataprepEditor/Source/DataprepEditor/Private/Widgets/DataprepGraph/SDataprepObjectSelectionFilter.h b/Engine/Plugins/Enterprise/DataprepEditor/Source/DataprepEditor/Private/Widgets/DataprepGraph/SDataprepObjectSelectionFilter.h
new file mode 100644
--- /dev/null
+++ b/Engine/Plugins/Enterprise/DataprepEditor/Source/DataprepEditor/Private/Widgets/DataprepGraph/SDataprepObjectSelectionFilter.h
@@ -0,0 +1,34 @@
+// Copyright Epic Games, Inc. All Rights Reserved.
+
+#pragma once
+
+#include "CoreMinimal.h"
+#include "Types/SlateEnums.h"
+#include "UObject/GCObject.h"
+#include "Widgets/DeclarativeSyntaxSupport.h"
+#include "Widgets/SCompoundWidget.h"
+
+class FMenuBuilder;
+class SComboButton;
+class UDataprepObjectSelectionFilter;
+
+template<class t>
+class SComboBox;
+
+struct FDataprepParametrizationActionData;
+
+class SDataprepObjectSelectionFilter : public SCompoundWidget, public FGCObject
+{
+	SLATE_BEGIN_ARGS(SDataprepObjectSelectionFilter) {}
+	SLATE_END_ARGS();
+
+	void Construct(const FArguments& InArgs, UDataprepObjectSelectionFilter& InFilter);
+
+private:
+
+	UDataprepObjectSelectionFilter* Filter;
+
+	//~ FGCObject interface
+	virtual void AddReferencedObjects(FReferenceCollector& Collector) override;
+	//~ End FGCObject interface
+};
