Added Nanite operator

diff --git a/Engine/Plugins/Enterprise/DataprepEditor/Source/DataprepLibraries/Private/DataprepLibrariesModule.cpp b/Engine/Plugins/Enterprise/DataprepEditor/Source/DataprepLibraries/Private/DataprepLibrariesModule.cpp
--- a/Engine/Plugins/Enterprise/DataprepEditor/Source/DataprepLibraries/Private/DataprepLibrariesModule.cpp
+++ b/Engine/Plugins/Enterprise/DataprepEditor/Source/DataprepLibraries/Private/DataprepLibrariesModule.cpp
@@ -25,7 +25,8 @@ public:
 		FPropertyEditorModule& PropertyModule = FModuleManager::LoadModuleChecked< FPropertyEditorModule >( TEXT("PropertyEditor") );
 		PropertyModule.RegisterCustomClassLayout( UDataprepSetLODGroupOperation::StaticClass()->GetFName(), FOnGetDetailCustomizationInstance::CreateStatic( &FDataprepSetLODGroupDetails::MakeDetails ) );
 		PropertyModule.RegisterCustomClassLayout( UDataprepSpawnActorsAtLocation::StaticClass()->GetFName(), FOnGetDetailCustomizationInstance::CreateStatic(&FDataprepSpawnActorsAtLocationDetails::MakeDetails ) );
-		PropertyModule.RegisterCustomClassLayout( UDataprepSetOutputFolder::StaticClass()->GetFName(), FOnGetDetailCustomizationInstance::CreateStatic( &FDataprepSetOutputFolderDetails::MakeDetails ) );
+		PropertyModule.RegisterCustomClassLayout(UDataprepSetOutputFolder::StaticClass()->GetFName(), FOnGetDetailCustomizationInstance::CreateStatic(&FDataprepSetOutputFolderDetails::MakeDetails));
+		PropertyModule.RegisterCustomClassLayout(UDataprepSetNaniteSettingsOperation::StaticClass()->GetFName(), FOnGetDetailCustomizationInstance::CreateStatic(&FDataprepSetNaniteSettingsDetails::MakeDetails));
 		RegisterMenus();
 	}
 
diff --git a/Engine/Plugins/Enterprise/DataprepEditor/Source/DataprepLibraries/Private/DataprepOperations.cpp b/Engine/Plugins/Enterprise/DataprepEditor/Source/DataprepLibraries/Private/DataprepOperations.cpp
--- a/Engine/Plugins/Enterprise/DataprepEditor/Source/DataprepLibraries/Private/DataprepOperations.cpp
+++ b/Engine/Plugins/Enterprise/DataprepEditor/Source/DataprepLibraries/Private/DataprepOperations.cpp
@@ -568,4 +568,150 @@ void UDataprepSetMaxTextureSizeOperation::PostEditChangeProperty(FPropertyChange
 	}
 }
 
+void UDataprepSetNaniteSettingsOperation::OnExecution_Implementation(const FDataprepContext& InContext)
+{
+#ifdef LOG_TIME
+	DataprepOperationTime::FTimeLogger TimeLogger(TEXT("SetNaniteSettings"), [&](FText Text) { this->LogInfo(Text); });
+#endif
+
+	// Execute operation
+	TArray<UObject*> ModifiedStaticMeshes;
+	UDataprepOperationsLibrary::SetNaniteSettings(InContext.Objects, bNaniteEnabled, PositionPrecision, PercentTriangles * 0.01f, ModifiedStaticMeshes);
+
+	if (ModifiedStaticMeshes.Num() > 0)
+	{
+		AssetsModified(MoveTemp(ModifiedStaticMeshes));
+	}
+}
+
+// Helper functions to replicate setting of PositionPrecision for Nanite settings
+// in StaticMesh editor.
+// Borrowed from FNaniteSettingsLayoutin StaticMeshEditorTool.cpp
+namespace NaniteSettingsUtils
+{
+	static const int32 DisplayPositionPrecisionMin = -6;
+	static const int32 DisplayPositionPrecisionMax = 13;
+	/** Strings representing valid PositionPrecision. */
+	static TArray< TSharedPtr< FString > > PositionPrecisionNames;
+
+	int32 PositionPrecisionIndexToValue(int32 Index)
+	{
+		check(Index >= 0);
+
+		if (Index == 0)
+		{
+			return MIN_int32;
+		}
+		else
+		{
+			int32 Value = DisplayPositionPrecisionMin + (Index - 1);
+			Value = FMath::Min(Value, DisplayPositionPrecisionMax);
+			return Value;
+		}
+	}
+
+	int32 PositionPrecisionValueToIndex(int32 Value)
+	{
+		if (Value == MIN_int32)
+		{
+			return 0;
+		}
+		else
+		{
+			Value = FMath::Clamp(Value, DisplayPositionPrecisionMin, DisplayPositionPrecisionMax);
+			return Value - DisplayPositionPrecisionMin + 1;
+		}
+	}
+
+	FString PositionPrecisionValueToDisplayString(int32 Value)
+	{
+		check(Value != MIN_int32);
+
+		if (Value <= 0)
+		{
+			return FString::Printf(TEXT("%dcm"), 1 << (-Value));
+		}
+		else
+		{
+			const float fValue = FMath::Exp2((double)-Value);
+			return FString::Printf(TEXT("1/%dcm (%.3gcm)"), 1 << Value, fValue);
+		}
+	}
+}
+
+void FDataprepSetNaniteSettingsDetails::OnPositionPrecisionChanged(TSharedPtr<FString> NewValue, ESelectInfo::Type /*SelectInfo*/)
+{
+	using namespace NaniteSettingsUtils;
+
+	int32 Index = PositionPrecisionNames.Find(NewValue);
+	if (Index != INDEX_NONE && PositionPrecisionPropertyHandle.IsValid())
+	{
+		PositionPrecisionPropertyHandle->SetValue(PositionPrecisionIndexToValue(Index));
+	}
+}
+
+TSharedRef< SWidget > FDataprepSetNaniteSettingsDetails::CreateWidget()
+{
+	using namespace NaniteSettingsUtils;
+
+	// Build list of PositionPrecision names the user will choose from
+	if (PositionPrecisionNames.Num() == 0)
+	{
+		PositionPrecisionNames.Reserve(20);
+
+		PositionPrecisionNames.Add(MakeShareable(new FString(TEXT("Auto"))));
+
+		for (int32 Value = DisplayPositionPrecisionMin; Value <= DisplayPositionPrecisionMax; ++Value)
+		{
+			PositionPrecisionNames.Add(MakeShareable(new FString(PositionPrecisionValueToDisplayString(Value))));
+		}
+	}
+
+	// Set displayed value to what is used in this operator's PositionPrecision
+	int32 SelectedIndex = PositionPrecisionValueToIndex(DataprepOperation->PositionPrecision);
+
+
+	// Create widget
+	return	SNew(STextComboBox)
+		.OptionsSource(&PositionPrecisionNames)
+		.InitiallySelectedItem(PositionPrecisionNames[SelectedIndex])
+		.OnSelectionChanged(this, &FDataprepSetNaniteSettingsDetails::OnPositionPrecisionChanged);
+}
+
+void FDataprepSetNaniteSettingsDetails::CustomizeDetails(IDetailLayoutBuilder& DetailBuilder)
+{
+	TArray< TWeakObjectPtr< UObject > > Objects;
+	DetailBuilder.GetObjectsBeingCustomized(Objects);
+	check(Objects.Num() > 0);
+
+	DataprepOperation = Cast< UDataprepSetNaniteSettingsOperation >(Objects[0].Get());
+	check(DataprepOperation);
+
+	TArray<FName> CategoryNames;
+	DetailBuilder.GetCategoryNames(CategoryNames);
+
+	FName CategoryName = CategoryNames.Num() > 0 ? CategoryNames[0] : FName(TEXT("MeshOperation_Internal"));
+	IDetailCategoryBuilder& ImportSettingsCategoryBuilder = DetailBuilder.EditCategory(CategoryName, FText::GetEmpty(), ECategoryPriority::Important);
+
+	PositionPrecisionPropertyHandle = DetailBuilder.GetProperty(GET_MEMBER_NAME_CHECKED(UDataprepSetNaniteSettingsOperation, PositionPrecision));
+
+	// Hide PositionPrecision property as it is replaced with custom widget
+	DetailBuilder.HideProperty(GET_MEMBER_NAME_CHECKED(UDataprepSetNaniteSettingsOperation, PositionPrecision));
+
+	FDetailWidgetRow& CustomAssetImportRow = ImportSettingsCategoryBuilder.AddCustomRow(FText::FromString(TEXT("Position Precision")));
+
+	CustomAssetImportRow.NameContent()
+		[
+			SNew(STextBlock)
+			.Text(LOCTEXT("DatasmithNaniteOperationsLabel", "PositionPrecision"))
+			.ToolTipText(LOCTEXT("DatasmithNaniteOperationsTooltip", "List of predefined position precision"))
+			.Font(DetailBuilder.GetDetailFont())
+		];
+
+	CustomAssetImportRow.ValueContent()
+		[
+			CreateWidget()
+		];
+}
+
 #undef LOCTEXT_NAMESPACE
diff --git a/Engine/Plugins/Enterprise/DataprepEditor/Source/DataprepLibraries/Private/DataprepOperations.h b/Engine/Plugins/Enterprise/DataprepEditor/Source/DataprepLibraries/Private/DataprepOperations.h
--- a/Engine/Plugins/Enterprise/DataprepEditor/Source/DataprepLibraries/Private/DataprepOperations.h
+++ b/Engine/Plugins/Enterprise/DataprepEditor/Source/DataprepLibraries/Private/DataprepOperations.h
@@ -569,3 +569,55 @@ protected:
 	// Track changes of MaxTextureSize and force values to be power of two
 	virtual void PostEditChangeProperty(FPropertyChangedEvent& PropertyChangedEvent) override;
 };
+
+UCLASS(Category = MeshOperation, Meta = (DisplayName = "Set Nanite Settings", ToolTip = "For each static mesh to process, enables/disables Nanite build and if enabled update settings."))
+class UDataprepSetNaniteSettingsOperation : public UDataprepOperation
+{
+	GENERATED_BODY()
+
+public:
+	/** If true, Nanite data will be generated */
+	UPROPERTY(EditAnywhere, BlueprintReadWrite, Category = MeshOperation, meta = (ToolTip = "If checked, Nanite data will be generated"))
+	bool bNaniteEnabled = false;
+
+	/** Position Precision. Step size is 2^(-PositionPrecision) cm. MIN_int32 is auto. */
+	// #ueent_todo: Same display as static mesh editor. See FNaniteSettingsLayout in StaticMeshEditorTool.h
+	UPROPERTY(EditAnywhere, BlueprintReadWrite, Category = MeshOperation_Internal, meta = (/*EditConditionHides, EditCondition = "bNaniteEnabled",*/ ToolTip = "Step size is 2^(-PositionPrecision) cm. MIN_int32 is auto"))
+	int32 PositionPrecision = MIN_int32;
+
+	/** Percentage of triangles to keep from LOD0. 100.0 = no reduction, 0.0 = no triangles. */
+	UPROPERTY(EditAnywhere, BlueprintReadWrite, Category = MeshOperation, meta = (/*EditConditionHides, EditCondition = "bNaniteEnabled",*/ ClampMin = 0, ClampMax = 100, UIMin = "0.0", UIMax = "100.0", ToolTip = "Percentage of triangles to keep from LOD0. 100.0 = no reduction, 0.0 = no triangles"))
+	float PercentTriangles = 100.f;
+
+	//~ Begin UDataprepOperation Interface
+public:
+	virtual FText GetCategory_Implementation() const override
+	{
+		return FDataprepOperationCategories::MeshOperation;
+	}
+
+protected:
+	virtual void OnExecution_Implementation(const FDataprepContext& InContext) override;
+	//~ End UDataprepOperation Interface
+};
+
+// Customization of the details of the DataprepSetNaniteSettingsOperation for the Dataprep editor.
+class FDataprepSetNaniteSettingsDetails : public IDetailCustomization
+{
+public:
+	static TSharedRef< IDetailCustomization > MakeDetails() { return MakeShared<FDataprepSetNaniteSettingsDetails>(); };
+
+	FDataprepSetNaniteSettingsDetails() : DataprepOperation(nullptr) {}
+
+	/** Called when details should be customized */
+	virtual void CustomizeDetails(IDetailLayoutBuilder& DetailBuilder) override;
+
+private:
+	TSharedRef< SWidget > CreateWidget();
+	void OnPositionPrecisionChanged(TSharedPtr<FString> NewValue, ESelectInfo::Type SelectInfo);
+
+private:
+	UDataprepSetNaniteSettingsOperation* DataprepOperation;
+
+	TSharedPtr<IPropertyHandle> PositionPrecisionPropertyHandle;
+};
diff --git a/Engine/Plugins/Enterprise/DataprepEditor/Source/DataprepLibraries/Private/DataprepOperationsLibrary.cpp b/Engine/Plugins/Enterprise/DataprepEditor/Source/DataprepLibraries/Private/DataprepOperationsLibrary.cpp
--- a/Engine/Plugins/Enterprise/DataprepEditor/Source/DataprepLibraries/Private/DataprepOperationsLibrary.cpp
+++ b/Engine/Plugins/Enterprise/DataprepEditor/Source/DataprepLibraries/Private/DataprepOperationsLibrary.cpp
@@ -69,6 +69,13 @@ void UDataprepOperationsLibrary::SetLods(const TArray<UObject*>& SelectedObjects
 
 void UDataprepOperationsLibrary::SetSimpleCollision(const TArray<UObject*>& SelectedObjects, const EScriptingCollisionShapeType_Deprecated ShapeType, TArray<UObject*>& ModifiedObjects)
 {
+	UStaticMeshEditorSubsystem* StaticMeshEditorSubsystem = GEditor->GetEditorSubsystem<UStaticMeshEditorSubsystem>();
+
+	if (!StaticMeshEditorSubsystem)
+	{
+		return;
+	}
+
 	TSet<UStaticMesh*> SelectedMeshes = DataprepOperationsLibraryUtil::GetSelectedMeshes(SelectedObjects);
 
 	// Make sure all static meshes to be processed have render data for NDOP types
@@ -99,13 +106,6 @@ void UDataprepOperationsLibrary::SetSimpleCollision(const TArray<UObject*>& Sele
 		{
 			DataprepOperationsLibraryUtil::FScopedStaticMeshEdit StaticMeshEdit( StaticMesh );
 
-			UStaticMeshEditorSubsystem* StaticMeshEditorSubsystem = GEditor->GetEditorSubsystem<UStaticMeshEditorSubsystem>();
-
-			if (!StaticMeshEditorSubsystem)
-			{
-				return;
-			}
-
 			// Remove existing simple collisions
 			StaticMeshEditorSubsystem->RemoveCollisionsWithNotification( StaticMesh, false );
 
@@ -138,10 +138,7 @@ void UDataprepOperationsLibrary::SetConvexDecompositionCollision(const TArray<UO
 
 	StaticMeshEditorSubsystem->BulkSetConvexDecompositionCollisionsWithNotification(StaticMeshes, HullCount, MaxHullVerts, HullPrecision, false);
 
-	for (UStaticMesh* StaticMesh : StaticMeshes)
-	{
-		ModifiedObjects.Add( StaticMesh );
-	}
+	ModifiedObjects.Append(StaticMeshes);
 }
 
 void UDataprepOperationsLibrary::SubstituteMaterial(const TArray<UObject*>& SelectedObjects, const FString& MaterialSearch, EEditorScriptingStringMatchType StringMatch, UMaterialInterface* MaterialSubstitute)
@@ -775,4 +772,31 @@ void UDataprepOperationsLibrary::ResizeTextures(const TArray<UTexture2D*>& InTex
 	}
 }
 
+void UDataprepOperationsLibrary::SetNaniteSettings(const TArray<UObject*>& SelectedObjects, bool bEnabled, int32 PositionPrecision, float PercentTriangles, TArray<UObject*>& ModifiedObjects)
+{
+#if WITH_EDITORONLY_DATA
+	TSet<UStaticMesh*> SelectedMeshes = DataprepOperationsLibraryUtil::GetSelectedMeshes(SelectedObjects);
+
+	FMeshNaniteSettings NewSettings;
+	NewSettings.bEnabled = bEnabled;
+	NewSettings.PositionPrecision = PositionPrecision;
+	NewSettings.PercentTriangles = FMath::Clamp(PercentTriangles, 0.f, 1.f);
+
+	// Apply Nanite settings but do not commit changes
+	TArray<UStaticMesh*> ModifiedMeshes;
+	ModifiedMeshes.Reserve(SelectedMeshes.Num());
+
+	for (UStaticMesh* StaticMesh : SelectedMeshes)
+	{
+		if (StaticMesh && StaticMesh->NaniteSettings != NewSettings)
+		{
+			StaticMesh->NaniteSettings = NewSettings;
+			ModifiedMeshes.Add(StaticMesh);
+		}
+	}
+
+	ModifiedObjects.Append(ModifiedMeshes);
+#endif // #if WITH_EDITORONLY_DATA
+}
+
 #undef LOCTEXT_NAMESPACE
diff --git a/Engine/Plugins/Enterprise/DataprepEditor/Source/DataprepLibraries/Private/DataprepOperationsLibrary.h b/Engine/Plugins/Enterprise/DataprepEditor/Source/DataprepLibraries/Private/DataprepOperationsLibrary.h
--- a/Engine/Plugins/Enterprise/DataprepEditor/Source/DataprepLibraries/Private/DataprepOperationsLibrary.h
+++ b/Engine/Plugins/Enterprise/DataprepEditor/Source/DataprepLibraries/Private/DataprepOperationsLibrary.h
@@ -337,6 +337,17 @@ public:
 	UFUNCTION(BlueprintCallable, Category = "Dataprep | Operation")
 	static void ResizeTextures(const TArray<UTexture2D*>& InTextures, int32 InMaxSize);
 
+	/**
+	 * Set Nanite settings for selected meshes
+	 * @param	InSelectedObjects			Array of objects to process.
+	 * @param	bInEnabled					If true, Nanite data will be generated.
+	 * @param	InPositionPrecision			Step size is 2^(-PositionPrecision) cm. MIN_int32 is auto.
+	 * @param	InPercentTriangles			Percentage of triangles to keep from LOD0. 1.0 = no reduction, 0.0 = no triangles.
+	 * @param	OutModifiedObjects			List of modified objects the processed meshes will be added to
+	 */
+	UFUNCTION(BlueprintCallable, Category = "Dataprep | Operation")
+	static void SetNaniteSettings(const TArray<UObject*>& InSelectedObjects, bool bInEnabled, int32 InPositionPrecision, float InPercentTriangles, TArray<UObject*>& OutModifiedObjects);
+
 private:
 	static void SubstituteMaterial(const TArray<UObject*>& SelectedObjects, const FString& MaterialSearch, EEditorScriptingStringMatchType StringMatch, const TArray<UMaterialInterface*>& MaterialList, UMaterialInterface* MaterialSubstitute);
 	static void SubstituteMesh(const TArray<UObject*>& SelectedObjects, const FString& MeshSearch, EEditorScriptingStringMatchType StringMatch, const TArray<UStaticMesh*>& MeshList, UStaticMesh* MeshSubstitute);
