Added display of meshes with Nanite enabled

diff --git a/Engine/Plugins/Enterprise/DataprepEditor/Source/DataprepEditor/Private/DataprepEditor.cpp b/Engine/Plugins/Enterprise/DataprepEditor/Source/DataprepEditor/Private/DataprepEditor.cpp
--- a/Engine/Plugins/Enterprise/DataprepEditor/Source/DataprepEditor/Private/DataprepEditor.cpp
+++ b/Engine/Plugins/Enterprise/DataprepEditor/Source/DataprepEditor/Private/DataprepEditor.cpp
@@ -1348,6 +1348,8 @@ TSharedRef<SDockTab> FDataprepEditor::SpawnTabStatistics(const FSpawnTabArgs & A
 	StatsWidget->AddStat(StatNameDrawCalls, CategoryScene, FDataprepStat(LOCTEXT("DataprepEditor_Stats_DrawCalls", "Draw Calls")));
 	StatsWidget->AddStat(FDataprepStats::StatNameTriangles, CategoryAssets, FDataprepStat(LOCTEXT("DataprepEditor_Stats_Triangles", "Triangles")));
 	StatsWidget->AddStat(FDataprepStats::StatNameVertices, CategoryAssets, FDataprepStat(LOCTEXT("DataprepEditor_Stats_Vertices", "Vertices")));
+	StatsWidget->AddStat(FDataprepStats::StatNameNaniteTriangles, CategoryAssets, FDataprepStat(LOCTEXT("DataprepEditor_Stats_NaniteTriangles", "Nanite Triangles")));
+	StatsWidget->AddStat(FDataprepStats::StatNameNaniteVertices, CategoryAssets, FDataprepStat(LOCTEXT("DataprepEditor_Stats_NaniteVertices", "Nanite Vertices")));
 	StatsWidget->AddStat(FDataprepStats::StatNameTextures, CategoryAssets, FDataprepStat(LOCTEXT("DataprepEditor_Stats_Textures", "Textures")));
 	StatsWidget->AddStat(FDataprepStats::StatNameTextureSize, CategoryAssets, FDataprepStat(LOCTEXT("DataprepEditor_Stats_TextureSize", "Max Texture Size")));
 	StatsWidget->AddStat(FDataprepStats::StatNameMeshes, CategoryAssets, FDataprepStat(LOCTEXT("DataprepEditor_Stats_Meshes", "Meshes")));
diff --git a/Engine/Plugins/Enterprise/DataprepEditor/Source/DataprepEditor/Private/DataprepStats.cpp b/Engine/Plugins/Enterprise/DataprepEditor/Source/DataprepEditor/Private/DataprepStats.cpp
--- a/Engine/Plugins/Enterprise/DataprepEditor/Source/DataprepEditor/Private/DataprepStats.cpp
+++ b/Engine/Plugins/Enterprise/DataprepEditor/Source/DataprepEditor/Private/DataprepStats.cpp
@@ -19,6 +19,8 @@
 
 const FName FDataprepStats::StatNameTriangles(TEXT("Triangles"));
 const FName FDataprepStats::StatNameVertices(TEXT("Vertices"));
+const FName FDataprepStats::StatNameNaniteTriangles(TEXT("NaniteTriangles"));
+const FName FDataprepStats::StatNameNaniteVertices(TEXT("NaniteVertices"));
 const FName FDataprepStats::StatNameTextures(TEXT("Textures"));
 const FName FDataprepStats::StatNameTextureSize(TEXT("TextureSize"));
 const FName FDataprepStats::StatNameMeshes(TEXT("Meshes"));
@@ -280,6 +282,16 @@ private:
 						Stats.AddCount(FDataprepStats::StatNameTriangles, StaticMeshSection.NumTriangles);
 						Stats.AddCount(FDataprepStats::StatNameVertices, RenderData->LODResources[0].GetNumVertices());
 					}
+
+					if (StaticMesh->NaniteSettings.bEnabled && StaticMesh->HasValidNaniteData())
+					{
+						const Nanite::FResources& Resources = StaticMesh->GetRenderData()->NaniteResources;
+						if (Resources.RootClusterPage.Num() > 0)
+						{
+							Stats.AddCount(FDataprepStats::StatNameNaniteTriangles, Resources.NumInputTriangles);
+							Stats.AddCount(FDataprepStats::StatNameNaniteVertices, Resources.NumInputVertices);
+						}
+					}
 				}
 			}
 			else if( SkeletalMeshComponent )
diff --git a/Engine/Plugins/Enterprise/DataprepEditor/Source/DataprepEditor/Private/DataprepStats.h b/Engine/Plugins/Enterprise/DataprepEditor/Source/DataprepEditor/Private/DataprepStats.h
--- a/Engine/Plugins/Enterprise/DataprepEditor/Source/DataprepEditor/Private/DataprepStats.h
+++ b/Engine/Plugins/Enterprise/DataprepEditor/Source/DataprepEditor/Private/DataprepStats.h
@@ -11,6 +11,8 @@ struct FDataprepStats
 	/**	Stat names */
 	static const FName StatNameTriangles;
 	static const FName StatNameVertices;
+	static const FName StatNameNaniteTriangles;
+	static const FName StatNameNaniteVertices;
 	static const FName StatNameTextures;
 	static const FName StatNameTextureSize;
 	static const FName StatNameMeshes;
@@ -26,6 +28,8 @@ struct FDataprepStats
 	{
 		StatsMap.Add(StatNameTriangles, 0);
 		StatsMap.Add(StatNameVertices, 0);
+		StatsMap.Add(StatNameNaniteTriangles, 0);
+		StatsMap.Add(StatNameNaniteVertices, 0);
 		StatsMap.Add(StatNameTextures, 0);
 		StatsMap.Add(StatNameTextureSize, 0);
 		StatsMap.Add(StatNameMeshes, 0);
diff --git a/Engine/Plugins/Enterprise/DataprepEditor/Source/DataprepEditor/Private/Widgets/SDataprepEditorViewport.cpp b/Engine/Plugins/Enterprise/DataprepEditor/Source/DataprepEditor/Private/Widgets/SDataprepEditorViewport.cpp
--- a/Engine/Plugins/Enterprise/DataprepEditor/Source/DataprepEditor/Private/Widgets/SDataprepEditorViewport.cpp
+++ b/Engine/Plugins/Enterprise/DataprepEditor/Source/DataprepEditor/Private/Widgets/SDataprepEditorViewport.cpp
@@ -335,6 +335,8 @@ void SDataprepEditorViewport::UpdateScene()
 	// Gather all static meshes used by actors in PreviewWorld
 	TArray< UStaticMeshComponent* > SceneMeshComponents = DataprepEditor3DPreviewUtils::GetComponentsFromWorld<UStaticMeshComponent>( WorldToPreview );
 
+	bCanShowNaniteProxyMenu = false;
+
 	if(SceneMeshComponents.Num() > 0)
 	{
 		int32 InvalidStaticMeshesCount = 0;
@@ -393,6 +395,8 @@ void SDataprepEditorViewport::UpdateScene()
 					const UStaticMesh* StaticMesh = SceneMeshComponent->GetStaticMesh();
 					const FTransform& ComponentToWorldTransform = SceneMeshComponent->GetComponentTransform();
 					SceneBounds += StaticMesh->GetExtendedBounds().GetBox().TransformBy( ComponentToWorldTransform );
+
+					bCanShowNaniteProxyMenu |= StaticMesh->NaniteSettings.bEnabled;
 				}
 			}
 
@@ -444,6 +448,8 @@ void SDataprepEditorViewport::UpdateScene()
 						FComponentReregisterContext ReregisterContext( PreviewMeshComponent );
 						PreviewMeshComponent->SetStaticMesh( StaticMesh );
 
+						PreviewMeshComponent->bDisplayNaniteProxyMesh = bShowNaniteProxyMenuChecked & StaticMesh->NaniteSettings.bEnabled;
+
 						FTransform ComponentToWorldTransform = SceneMeshComponent->GetComponentTransform();
 
 						PreviewMeshComponent->AttachToComponent( PreviewActor->GetRootComponent(), FAttachmentTransformRules::KeepRelativeTransform );
@@ -608,6 +614,8 @@ void SDataprepEditorViewport::UpdateOverlayText()
 
 	TSet<UStaticMesh*> StaticMeshes;
 
+	int32 NaniteTrianglesCount = 0;
+	int32 NaniteVerticesCount = 0;
 	int32 TrianglesCount = 0;
 	int32 VerticesCount = 0;
 	for( const TWeakObjectPtr< UStaticMeshComponent >& PreviewMeshComponent : PreviewMeshComponents )
@@ -615,17 +623,57 @@ void SDataprepEditorViewport::UpdateOverlayText()
 		if(UStaticMeshComponent* MeshComponent = PreviewMeshComponent.Get())
 		{
 			UStaticMesh* StaticMesh = MeshComponent->GetStaticMesh();
-			TrianglesCount += StaticMesh->GetRenderData()->LODResources[0].GetNumTriangles();
-			VerticesCount += StaticMesh->GetRenderData()->LODResources[0].GetNumVertices();
+
+			FStaticMeshLODResources& LODResource = StaticMesh->GetRenderData()->LODResources[0];
+
+			if (StaticMesh->NaniteSettings.bEnabled && StaticMesh->HasValidNaniteData())
+			{
+				const Nanite::FResources& Resources = StaticMesh->GetRenderData()->NaniteResources;
+				if (Resources.RootClusterPage.Num() > 0)
+				{
+					NaniteTrianglesCount += bShowNaniteProxyMenuChecked ? LODResource.GetNumTriangles() : Resources.NumInputTriangles;
+					NaniteVerticesCount += bShowNaniteProxyMenuChecked ? LODResource.GetNumVertices() : Resources.NumInputVertices;
+				}
+			}
+			else
+			{
+				TrianglesCount += LODResource.GetNumTriangles();
+				VerticesCount += LODResource.GetNumVertices();
+			}
 			StaticMeshes.Add(StaticMesh);
 		}
 	}
 
+	if (NaniteTrianglesCount > 0)
+	{
+		static FText ShowingNaniteProxy = LOCTEXT("ShowingNaniteProxy", "(Showing Proxy)");
+
+		TextItems.Add(FOverlayTextItem(
+			FText::Format(LOCTEXT("NaniteEnabled", "Nanite Enabled {0}"), bShowNaniteProxyMenuChecked ? ShowingNaniteProxy : FText::GetEmpty())));
+
+		if (bShowNaniteProxyMenuChecked)
+		{
+			TextItems.Add(FOverlayTextItem(
+				FText::Format(LOCTEXT("Proxy_Triangles", "#Proxy_Triangles: {0}"), FText::AsNumber(NaniteTrianglesCount))));
+
+			TextItems.Add(FOverlayTextItem(
+				FText::Format(LOCTEXT("Proxy_Vertices", "#Proxy_Vertices: {0}"), FText::AsNumber(NaniteVerticesCount))));
+		}
+		else
+		{
+			TextItems.Add(FOverlayTextItem(
+				FText::Format(LOCTEXT("Nanite_Triangles", "#Nanite_Triangles: {0}"), FText::AsNumber(NaniteTrianglesCount))));
+
+			TextItems.Add(FOverlayTextItem(
+				FText::Format(LOCTEXT("Nanite_Vertices", "#Nanite_Vertices: {0}"), FText::AsNumber(NaniteVerticesCount))));
+		}
+	}
+
 	TextItems.Add(FOverlayTextItem(
-		FText::Format(LOCTEXT( "Triangles_F", "#Triangles:  {0}"), FText::AsNumber(TrianglesCount))));
+		FText::Format(LOCTEXT("Triangles_F", "#Triangles: {0}"), FText::AsNumber(TrianglesCount))));
 
 	TextItems.Add(FOverlayTextItem(
-		FText::Format(LOCTEXT( "Vertices_F", "#Vertices:  {0}"), FText::AsNumber(VerticesCount))));
+		FText::Format(LOCTEXT("Vertices_F", "#Vertices: {0}"), FText::AsNumber(VerticesCount))));
 
 	FVector SceneExtents = SceneBounds.GetExtent();
 	TextItems.Add(FOverlayTextItem(
@@ -697,7 +745,14 @@ void SDataprepEditorViewport::BindCommands()
 		Commands.SetShowBounds,
 		FExecuteAction::CreateSP( EditorViewportClientRef, &FEditorViewportClient::ToggleShowBounds ),
 		FCanExecuteAction(),
-		FIsActionChecked::CreateSP( EditorViewportClientRef, &FEditorViewportClient::IsSetShowBoundsChecked ) );
+		FIsActionChecked::CreateSP(EditorViewportClientRef, &FEditorViewportClient::IsSetShowBoundsChecked ) );
+
+	CommandList->MapAction(
+		Commands.SetShowNaniteProxy,
+		FExecuteAction::CreateSP(this, &SDataprepEditorViewport::ToggleShowNaniteProxy),
+		FCanExecuteAction(),
+		FIsActionChecked::CreateSP(this, &SDataprepEditorViewport::IsSetShowNaniteProxyChecked),
+		FIsActionButtonVisible::CreateSP(this, &SDataprepEditorViewport::IsShowNaniteProxyVisible));
 
 	CommandList->MapAction(
 		Commands.ApplyOriginalMaterial,
@@ -1255,6 +1310,26 @@ void SDataprepEditorViewport::LoadDefaultSettings()
 	}
 }
 
+void SDataprepEditorViewport::SetShowNaniteProxy(bool bShow)
+{
+	bShowNaniteProxyMenuChecked = bShow;
+
+	for (const TWeakObjectPtr< UStaticMeshComponent >& PreviewMeshComponentPtr : PreviewMeshComponents)
+	{
+		if (UCustomStaticMeshComponent* PreviewMeshComponent = Cast<UCustomStaticMeshComponent>(PreviewMeshComponentPtr.Get()))
+		{
+			const UStaticMesh* StaticMesh = PreviewMeshComponent->GetStaticMesh();
+			if (StaticMesh && StaticMesh->NaniteSettings.bEnabled)
+			{
+				PreviewMeshComponent->bDisplayNaniteProxyMesh = bShowNaniteProxyMenuChecked & StaticMesh->NaniteSettings.bEnabled;
+				PreviewMeshComponent->MarkRenderStateDirty();
+			}
+		}
+	}
+
+	UpdateOverlayText();
+}
+
 //
 // FDataprepEditorViewportClient Class
 //
@@ -1413,6 +1488,7 @@ TSharedRef<SWidget> SDataprepEditorViewportToolbar::GenerateShowMenu() const
 	{
 		ShowMenuBuilder.AddMenuEntry(FDataprepEditorViewportCommands::Get().SetShowGrid);
 		ShowMenuBuilder.AddMenuEntry(FDataprepEditorViewportCommands::Get().SetShowBounds);
+		ShowMenuBuilder.AddMenuEntry(FDataprepEditorViewportCommands::Get().SetShowNaniteProxy);
 	}
 
 	// #ueent_remark: Look at SAnimViewportToolBar::GenerateShowMenu in SAnimViewportToolBar.cpp for adding ShowFlagFilter to the Show menu
@@ -1546,8 +1622,9 @@ void SDataprepEditorViewport::ToggleShowOrientedBox()
 void FDataprepEditorViewportCommands::RegisterCommands()
 {
 	// Show menu
-	UI_COMMAND(SetShowGrid, "Grid", "Displays the viewport grid.", EUserInterfaceActionType::ToggleButton, FInputChord());
-	UI_COMMAND(SetShowBounds, "Bounds", "Toggles display of the bounds of the selected component.", EUserInterfaceActionType::ToggleButton, FInputChord());
+	UI_COMMAND(SetShowGrid, "Grid", "Displays the viewport grid.", EUserInterfaceActionType::ToggleButton, FInputChord(EModifierKey::Control, EKeys::G));
+	UI_COMMAND(SetShowBounds, "Bounds", "Toggles display of the bounds of the selected component.", EUserInterfaceActionType::ToggleButton, FInputChord(EModifierKey::Control, EKeys::B));
+	UI_COMMAND(SetShowNaniteProxy, "Nanite Proxy", "Toggles the display of the Nanite proxy mesh.", EUserInterfaceActionType::ToggleButton, FInputChord(EModifierKey::Control, EKeys::N));
 
 	// Rendering Material
 	UI_COMMAND(ApplyOriginalMaterial, "None", "Display all meshes with original materials.", EUserInterfaceActionType::RadioButton, FInputChord());
@@ -1577,6 +1654,13 @@ FPrimitiveSceneProxy* UCustomStaticMeshComponent::CreateSceneProxy()
 		return nullptr;
 	}
 
+	// Do not overwrite proxy if Nanite supported and there is built Nanite data for the static mesh
+	// #ueent_todo: Track changes to Nanite::FSceneProxy code in case wireframe display is supported
+	if (ShouldCreateNaniteProxy() && !bDisplayNaniteProxyMesh)
+	{
+		return UStaticMeshComponent::CreateSceneProxy();
+	}
+
 	const FStaticMeshLODResourcesArray& LODResources = GetStaticMesh()->GetRenderData()->LODResources;
 	if (LODResources.Num() == 0	|| LODResources[FMath::Clamp<int32>(GetStaticMesh()->GetDefaultMinLOD(), 0, LODResources.Num()-1)].VertexBuffers.StaticMeshVertexBuffer.GetNumVertices() == 0)
 	{
diff --git a/Engine/Plugins/Enterprise/DataprepEditor/Source/DataprepEditor/Private/Widgets/SDataprepEditorViewport.h b/Engine/Plugins/Enterprise/DataprepEditor/Source/DataprepEditor/Private/Widgets/SDataprepEditorViewport.h
--- a/Engine/Plugins/Enterprise/DataprepEditor/Source/DataprepEditor/Private/Widgets/SDataprepEditorViewport.h
+++ b/Engine/Plugins/Enterprise/DataprepEditor/Source/DataprepEditor/Private/Widgets/SDataprepEditorViewport.h
@@ -79,6 +79,7 @@ public:
 public:
 	TSharedPtr<FUICommandInfo> SetShowGrid;
 	TSharedPtr<FUICommandInfo> SetShowBounds;
+	TSharedPtr<FUICommandInfo> SetShowNaniteProxy;
 
 	TSharedPtr<FUICommandInfo> ApplyOriginalMaterial;
 	TSharedPtr<FUICommandInfo> ApplyBackFaceMaterial;
@@ -307,6 +308,18 @@ private:
 	/** Returns true if the given primitive component is part of the selection set */
 	bool IsComponentSelected(const UPrimitiveComponent* PrimitiveComponent);
 
+	/** Sets the show Nanite flag */
+	void SetShowNaniteProxy(bool bShow);
+
+	/** Callback for toggling the show Nanite flag. */
+	void ToggleShowNaniteProxy() { SetShowNaniteProxy(!bShowNaniteProxyMenuChecked); }
+
+	/** Callback for checking the Nanite show flag. */
+	bool IsSetShowNaniteProxyChecked() const { return bShowNaniteProxyMenuChecked; }
+
+	/** Callback for checking whether the Nanite show entry should be displayed. */
+	bool IsShowNaniteProxyVisible() const { return bCanShowNaniteProxyMenu; }
+
 private:
 	/** The scene for this viewport */
 	TSharedPtr< FAdvancedPreviewScene > PreviewScene;
@@ -395,5 +408,8 @@ private:
 	bool bShowOrientedBox;
 #endif
 
+	bool bCanShowNaniteProxyMenu = false;
+	bool bShowNaniteProxyMenuChecked = false;
+
 	friend FDataprepEditorViewportClient;
 };
diff --git a/Engine/Plugins/Enterprise/DataprepEditor/Source/DataprepEditor/Private/Widgets/SDataprepStats.cpp b/Engine/Plugins/Enterprise/DataprepEditor/Source/DataprepEditor/Private/Widgets/SDataprepStats.cpp
--- a/Engine/Plugins/Enterprise/DataprepEditor/Source/DataprepEditor/Private/Widgets/SDataprepStats.cpp
+++ b/Engine/Plugins/Enterprise/DataprepEditor/Source/DataprepEditor/Private/Widgets/SDataprepStats.cpp
@@ -194,6 +194,8 @@ void SDataprepStats::SetStats(const TSharedPtr<FDataprepStats>& InStats, bool bI
 
 	SetSingleStat(FDataprepStats::StatNameTriangles);
 	SetSingleStat(FDataprepStats::StatNameVertices);
+	SetSingleStat(FDataprepStats::StatNameNaniteTriangles);
+	SetSingleStat(FDataprepStats::StatNameNaniteVertices);
 	SetSingleStat(FDataprepStats::StatNameTextures);
 	SetSingleStat(FDataprepStats::StatNameTextureSize);
 	SetSingleStat(FDataprepStats::StatNameMeshes);
