Fixed warning from static analysis

diff --git a/Engine/Plugins/Enterprise/DataprepEditor/Source/DataprepEditor/Private/Widgets/SDataprepEditorViewport.cpp b/Engine/Plugins/Enterprise/DataprepEditor/Source/DataprepEditor/Private/Widgets/SDataprepEditorViewport.cpp
--- a/Engine/Plugins/Enterprise/DataprepEditor/Source/DataprepEditor/Private/Widgets/SDataprepEditorViewport.cpp
+++ b/Engine/Plugins/Enterprise/DataprepEditor/Source/DataprepEditor/Private/Widgets/SDataprepEditorViewport.cpp
@@ -448,7 +448,7 @@ void SDataprepEditorViewport::UpdateScene()
 						FComponentReregisterContext ReregisterContext( PreviewMeshComponent );
 						PreviewMeshComponent->SetStaticMesh( StaticMesh );
 
-						PreviewMeshComponent->bDisplayNaniteProxyMesh = bShowNaniteProxyMenuChecked & StaticMesh->NaniteSettings.bEnabled;
+						PreviewMeshComponent->bDisplayNaniteProxyMesh = bShowNaniteProxyMenuChecked && StaticMesh->NaniteSettings.bEnabled;
 
 						FTransform ComponentToWorldTransform = SceneMeshComponent->GetComponentTransform();
 
@@ -1321,7 +1321,7 @@ void SDataprepEditorViewport::SetShowNaniteProxy(bool bShow)
 			const UStaticMesh* StaticMesh = PreviewMeshComponent->GetStaticMesh();
 			if (StaticMesh && StaticMesh->NaniteSettings.bEnabled)
 			{
-				PreviewMeshComponent->bDisplayNaniteProxyMesh = bShowNaniteProxyMenuChecked & StaticMesh->NaniteSettings.bEnabled;
+				PreviewMeshComponent->bDisplayNaniteProxyMesh = bShowNaniteProxyMenuChecked && StaticMesh->NaniteSettings.bEnabled;
 				PreviewMeshComponent->MarkRenderStateDirty();
 			}
 		}
