Fixing CIS errors
#rnx

diff --git a/Engine/Plugins/Enterprise/DataprepEditor/Source/DataprepEditor/Private/DataprepStats.cpp b/Engine/Plugins/Enterprise/DataprepEditor/Source/DataprepEditor/Private/DataprepStats.cpp
--- a/Engine/Plugins/Enterprise/DataprepEditor/Source/DataprepEditor/Private/DataprepStats.cpp
+++ b/Engine/Plugins/Enterprise/DataprepEditor/Source/DataprepEditor/Private/DataprepStats.cpp
@@ -286,7 +286,7 @@ private:
 					if (StaticMesh->NaniteSettings.bEnabled && StaticMesh->HasValidNaniteData())
 					{
 						const Nanite::FResources& Resources = StaticMesh->GetRenderData()->NaniteResources;
-						if (Resources.RootClusterPage.Num() > 0)
+						if (Resources.RootData.Num() > 0)
 						{
 							Stats.AddCount(FDataprepStats::StatNameNaniteTriangles, Resources.NumInputTriangles);
 							Stats.AddCount(FDataprepStats::StatNameNaniteVertices, Resources.NumInputVertices);
diff --git a/Engine/Plugins/Enterprise/DataprepEditor/Source/DataprepEditor/Private/Widgets/SDataprepEditorViewport.cpp b/Engine/Plugins/Enterprise/DataprepEditor/Source/DataprepEditor/Private/Widgets/SDataprepEditorViewport.cpp
--- a/Engine/Plugins/Enterprise/DataprepEditor/Source/DataprepEditor/Private/Widgets/SDataprepEditorViewport.cpp
+++ b/Engine/Plugins/Enterprise/DataprepEditor/Source/DataprepEditor/Private/Widgets/SDataprepEditorViewport.cpp
@@ -629,7 +629,7 @@ void SDataprepEditorViewport::UpdateOverlayText()
 			if (StaticMesh->NaniteSettings.bEnabled && StaticMesh->HasValidNaniteData())
 			{
 				const Nanite::FResources& Resources = StaticMesh->GetRenderData()->NaniteResources;
-				if (Resources.RootClusterPage.Num() > 0)
+				if (Resources.RootData.Num() > 0)
 				{
 					NaniteTrianglesCount += bShowNaniteProxyMenuChecked ? LODResource.GetNumTriangles() : Resources.NumInputTriangles;
 					NaniteVerticesCount += bShowNaniteProxyMenuChecked ? LODResource.GetNumVertices() : Resources.NumInputVertices;
