Dataprep - Scene preview actors are no longer 'right-clickable'

diff --git a/Engine/Plugins/Enterprise/DataprepEditor/Source/DataprepEditor/Private/DataprepEditorOutliner.cpp b/Engine/Plugins/Enterprise/DataprepEditor/Source/DataprepEditor/Private/DataprepEditorOutliner.cpp
--- a/Engine/Plugins/Enterprise/DataprepEditor/Source/DataprepEditor/Private/DataprepEditorOutliner.cpp
+++ b/Engine/Plugins/Enterprise/DataprepEditor/Source/DataprepEditor/Private/DataprepEditorOutliner.cpp
@@ -373,7 +373,6 @@ void FDataprepEditor::CreateScenePreviewTab()
 
 	FSceneOutlinerInitializationOptions SceneOutlinerOptions;
 	SceneOutlinerOptions.ModeFactory = ModeFactory;
-	//SceneOutlinerOptions.ContextMenuOverride.BindRaw(this, &FDataprepEditor::OnSceneOutlinerContextMenuOpening); // TODO: ContextMenuOverride doesn't exist anymore
 
 	SceneOutliner = SceneOutlinerModule.CreateSceneOutliner(SceneOutlinerOptions);
 
@@ -415,7 +414,7 @@ void FDataprepEditor::CreateScenePreviewTab()
 		];
 }
 
-TSharedPtr<SWidget> FDataprepEditor::OnSceneOutlinerContextMenuOpening()
+TSharedPtr<SWidget> FDataprepEditorOutlinerMode::CreateContextMenu()
 {
 	DataprepEditorSceneOutlinerUtils::FGetSelectionFromSceneOutliner Selector;
 
@@ -434,7 +433,9 @@ TSharedPtr<SWidget> FDataprepEditor::OnSceneOutlinerContextMenuOpening()
 		}
 	}
 
-	if (SelectedActors.Num() == 0)
+	TSharedPtr<FDataprepEditor> DataprepEditor = DataprepEditorPtr.Pin();
+
+	if (!DataprepEditor.IsValid() || SelectedActors.Num() == 0)
 	{
 		return TSharedPtr<SWidget>();
 	}
@@ -443,7 +444,7 @@ TSharedPtr<SWidget> FDataprepEditor::OnSceneOutlinerContextMenuOpening()
 
 	UDataprepEditorContextMenuContext* ContextObject = NewObject<UDataprepEditorContextMenuContext>();
 	ContextObject->SelectedObjects = SelectedActors.Array();
-	ContextObject->DataprepAsset = GetDataprepAsset();
+	ContextObject->DataprepAsset = DataprepEditor->GetDataprepAsset();
 
 	UToolMenus* ToolMenus = UToolMenus::Get();
 	FToolMenuContext MenuContext( nullptr, nullptr, ContextObject );
diff --git a/Engine/Plugins/Enterprise/DataprepEditor/Source/DataprepEditor/Private/DataprepEditorOutlinerMode.h b/Engine/Plugins/Enterprise/DataprepEditor/Source/DataprepEditor/Private/DataprepEditorOutlinerMode.h
--- a/Engine/Plugins/Enterprise/DataprepEditor/Source/DataprepEditor/Private/DataprepEditorOutlinerMode.h
+++ b/Engine/Plugins/Enterprise/DataprepEditor/Source/DataprepEditor/Private/DataprepEditorOutlinerMode.h
@@ -69,6 +69,8 @@ public:
 		}
 	}
 
+	virtual TSharedPtr<SWidget> CreateContextMenu() override;
+
 private:
 	void OnActorLabelChanged(AActor* InChangedActor)
 	{
