Dataprep operation "Replace Asset References" should not execute on actors

diff --git a/Engine/Plugins/Enterprise/DataprepEditor/Source/DataprepLibraries/Private/DataprepOperationsLibrary.cpp b/Engine/Plugins/Enterprise/DataprepEditor/Source/DataprepLibraries/Private/DataprepOperationsLibrary.cpp
--- a/Engine/Plugins/Enterprise/DataprepEditor/Source/DataprepLibraries/Private/DataprepOperationsLibrary.cpp
+++ b/Engine/Plugins/Enterprise/DataprepEditor/Source/DataprepLibraries/Private/DataprepOperationsLibrary.cpp
@@ -462,6 +462,14 @@ void UDataprepOperationsLibrary::ConsolidateObjects(const TArray< UObject* >& Se
 	UObject* ObjectToConsolidateTo = SelectedObjects[0];
 	check(ObjectToConsolidateTo);
 
+	UObject* Outer = ObjectToConsolidateTo->GetOuter();
+
+	if (nullptr == Outer || !Outer->IsA(UPackage::StaticClass()))
+	{
+		UE_LOG(LogDataprep, Warning, TEXT("Consolidate failed: the object %s is not an asset"), *ObjectToConsolidateTo->GetName());
+		return;
+	}
+
 	const UClass* ComparisonClass = ObjectToConsolidateTo->GetClass();
 	check(ComparisonClass);
 
