Dataprep - review deprecated code and api calls

diff --git a/Engine/Plugins/Enterprise/DataprepEditor/Source/DataprepLibraries/Private/DataprepOperations.cpp b/Engine/Plugins/Enterprise/DataprepEditor/Source/DataprepLibraries/Private/DataprepOperations.cpp
--- a/Engine/Plugins/Enterprise/DataprepEditor/Source/DataprepLibraries/Private/DataprepOperations.cpp
+++ b/Engine/Plugins/Enterprise/DataprepEditor/Source/DataprepLibraries/Private/DataprepOperations.cpp
@@ -74,7 +74,7 @@ void UDataprepSetLODsOperation::OnExecution_Implementation(const FDataprepContex
 	}
 
 	// Fill up mesh reduction struct
-	FEditorScriptingMeshReductionOptions_Deprecated ReductionOptions;
+	FStaticMeshReductionOptions ReductionOptions;
 
 	ReductionOptions.bAutoComputeLODScreenSize = bAutoComputeLODScreenSize;
 	ReductionOptions.ReductionSettings.SetNum( LODCount );
diff --git a/Engine/Plugins/Enterprise/DataprepEditor/Source/DataprepLibraries/Private/DataprepOperations.h b/Engine/Plugins/Enterprise/DataprepEditor/Source/DataprepLibraries/Private/DataprepOperations.h
--- a/Engine/Plugins/Enterprise/DataprepEditor/Source/DataprepLibraries/Private/DataprepOperations.h
+++ b/Engine/Plugins/Enterprise/DataprepEditor/Source/DataprepLibraries/Private/DataprepOperations.h
@@ -98,14 +98,14 @@ class UDataprepSetSimpleCollisionOperation : public UDataprepOperation
 	GENERATED_BODY()
 
 	UDataprepSetSimpleCollisionOperation()
-		: ShapeType(EScriptingCollisionShapeType_Deprecated::Box)
+		: ShapeType(EScriptCollisionShapeType::Box)
 	{
 	}
 
 public:
 	// Shape's of the collision geometry encompassing the static mesh
 	UPROPERTY(EditAnywhere, BlueprintReadWrite, Category = MeshOperation, meta = (ToolTip = "Shape's of the collision geometry encompassing the static mesh"))
-		EScriptingCollisionShapeType_Deprecated ShapeType;
+		EScriptCollisionShapeType ShapeType;
 
 	//~ Begin UDataprepOperation Interface
 public:
diff --git a/Engine/Plugins/Enterprise/DataprepEditor/Source/DataprepLibraries/Private/DataprepOperationsLibrary.cpp b/Engine/Plugins/Enterprise/DataprepEditor/Source/DataprepLibraries/Private/DataprepOperationsLibrary.cpp
--- a/Engine/Plugins/Enterprise/DataprepEditor/Source/DataprepLibraries/Private/DataprepOperationsLibrary.cpp
+++ b/Engine/Plugins/Enterprise/DataprepEditor/Source/DataprepLibraries/Private/DataprepOperationsLibrary.cpp
@@ -42,7 +42,7 @@ DEFINE_LOG_CATEGORY(LogDataprep);
 
 extern UNREALED_API UEditorEngine* GEditor;
 
-void UDataprepOperationsLibrary::SetLods(const TArray<UObject*>& SelectedObjects, const FEditorScriptingMeshReductionOptions_Deprecated& ReductionOptions, TArray<UObject*>& ModifiedObjects)
+void UDataprepOperationsLibrary::SetLods(const TArray<UObject*>& SelectedObjects, const FStaticMeshReductionOptions& ReductionOptions, TArray<UObject*>& ModifiedObjects)
 {
 	TSet<UStaticMesh*> SelectedMeshes = DataprepOperationsLibraryUtil::GetSelectedMeshes(SelectedObjects);
 
@@ -60,14 +60,14 @@ void UDataprepOperationsLibrary::SetLods(const TArray<UObject*>& SelectedObjects
 		{
 			DataprepOperationsLibraryUtil::FScopedStaticMeshEdit StaticMeshEdit( StaticMesh );
 
-			StaticMeshEditorSubsystem->SetLodsWithNotification(StaticMesh, UDEPRECATED_EditorStaticMeshLibrary::ConvertReductionOptions(ReductionOptions), false);
+			StaticMeshEditorSubsystem->SetLodsWithNotification(StaticMesh, ReductionOptions, false);
 
 			ModifiedObjects.Add( StaticMesh );
 		}
 	}
 }
 
-void UDataprepOperationsLibrary::SetSimpleCollision(const TArray<UObject*>& SelectedObjects, const EScriptingCollisionShapeType_Deprecated ShapeType, TArray<UObject*>& ModifiedObjects)
+void UDataprepOperationsLibrary::SetSimpleCollision(const TArray<UObject*>& SelectedObjects, const EScriptCollisionShapeType ShapeType, TArray<UObject*>& ModifiedObjects)
 {
 	UStaticMeshEditorSubsystem* StaticMeshEditorSubsystem = GEditor->GetEditorSubsystem<UStaticMeshEditorSubsystem>();
 
@@ -82,11 +82,11 @@ void UDataprepOperationsLibrary::SetSimpleCollision(const TArray<UObject*>& Sele
 	bool bNeedRenderData = false;
 	switch (ShapeType)
 	{
-		case EScriptingCollisionShapeType_Deprecated::NDOP10_X:
-		case EScriptingCollisionShapeType_Deprecated::NDOP10_Y:
-		case EScriptingCollisionShapeType_Deprecated::NDOP10_Z:
-		case EScriptingCollisionShapeType_Deprecated::NDOP18:
-		case EScriptingCollisionShapeType_Deprecated::NDOP26:
+		case EScriptCollisionShapeType::NDOP10_X:
+		case EScriptCollisionShapeType::NDOP10_Y:
+		case EScriptCollisionShapeType::NDOP10_Z:
+		case EScriptCollisionShapeType::NDOP18:
+		case EScriptCollisionShapeType::NDOP26:
 		{
 			bNeedRenderData = true;
 			break;
@@ -109,7 +109,7 @@ void UDataprepOperationsLibrary::SetSimpleCollision(const TArray<UObject*>& Sele
 			// Remove existing simple collisions
 			StaticMeshEditorSubsystem->RemoveCollisionsWithNotification( StaticMesh, false );
 
-			StaticMeshEditorSubsystem->AddSimpleCollisionsWithNotification( StaticMesh, UDEPRECATED_EditorStaticMeshLibrary::ConvertCollisionShape(ShapeType), false );
+			StaticMeshEditorSubsystem->AddSimpleCollisionsWithNotification( StaticMesh, ShapeType, false );
 
 			ModifiedObjects.Add( StaticMesh );
 		}
diff --git a/Engine/Plugins/Enterprise/DataprepEditor/Source/DataprepLibraries/Private/DataprepOperationsLibrary.h b/Engine/Plugins/Enterprise/DataprepEditor/Source/DataprepLibraries/Private/DataprepOperationsLibrary.h
--- a/Engine/Plugins/Enterprise/DataprepEditor/Source/DataprepLibraries/Private/DataprepOperationsLibrary.h
+++ b/Engine/Plugins/Enterprise/DataprepEditor/Source/DataprepLibraries/Private/DataprepOperationsLibrary.h
@@ -144,7 +144,7 @@ public:
 	 * Calls UEditorStaticMeshLibrary::SetLods on each static mesh of the resulting array.
 	 */
 	UFUNCTION(BlueprintCallable, Category = "Dataprep | Operation")
-	static void SetLods(const TArray<UObject*>& SelectedObjects, const FEditorScriptingMeshReductionOptions_Deprecated& ReductionOptions, TArray<UObject*>& ModifiedObjects);
+	static void SetLods(const TArray<UObject*>& SelectedObjects, const FStaticMeshReductionOptions& ReductionOptions, TArray<UObject*>& ModifiedObjects);
 
 	/**
 	 * Set one simple collision of the given shape type on the static meshes contained in the
@@ -160,7 +160,7 @@ public:
 	 * Calls UEditorStaticMeshLibrary::AddSimpleCollisions on each static mesh of the resulting array.
 	 */
 	UFUNCTION(BlueprintCallable, Category = "Dataprep | Operation")
-	static void SetSimpleCollision(const TArray<UObject*>& SelectedObjects, const EScriptingCollisionShapeType_Deprecated ShapeType, TArray<UObject*>& ModifiedObjects);
+	static void SetSimpleCollision(const TArray<UObject*>& SelectedObjects, const EScriptCollisionShapeType ShapeType, TArray<UObject*>& ModifiedObjects);
 
 	/**
 	 * Add complex collision on the static meshes contained in the input array
