Dataprep: Editor crashes after clicking on drop down in the set metadata operator in the dataprep asset

diff --git a/Engine/Plugins/Enterprise/DataprepEditor/Source/DataprepEditor/Private/Widgets/DataprepWidgets.cpp b/Engine/Plugins/Enterprise/DataprepEditor/Source/DataprepEditor/Private/Widgets/DataprepWidgets.cpp
--- a/Engine/Plugins/Enterprise/DataprepEditor/Source/DataprepEditor/Private/Widgets/DataprepWidgets.cpp
+++ b/Engine/Plugins/Enterprise/DataprepEditor/Source/DataprepEditor/Private/Widgets/DataprepWidgets.cpp
@@ -445,9 +445,11 @@ void SDataprepDetailsView::OnObjectReplaced(const TMap<UObject*, UObject*>& Repl
 
 void SDataprepDetailsView::ForceRefresh()
 {
-	// ueent_hotfix Hotfix for 4.24 (Remove the ui flickering)
-	Invalidate(EInvalidateWidgetReason::Prepass);
-	bRefreshObjectToDisplay = true;
+	if (Generator.IsValid())
+	{
+		Generator->SetObjects({ DetailedObject });
+	}
+	Construct();
 }
 
 void SDataprepDetailsView::OnDataprepParameterizationStatusForObjectsChanged(const TSet<UObject*>* Objects)
@@ -631,7 +633,6 @@ void SDataprepDetailsView::AddWidgets( const TArray< TSharedRef< IDetailTreeNode
 
 void SDataprepDetailsView::Construct(const FArguments& InArgs)
 {
-	bRefreshObjectToDisplay = false;
 	DetailedObject = InArgs._Object;
 	Spacing = InArgs._Spacing;
 	bColumnPadding = InArgs._ColumnPadding;
@@ -672,7 +673,6 @@ void SDataprepDetailsView::Construct()
 {
 	TrackedProperties.Reset();
 	DataprepAssetForParameterization.Reset();
-	bHasCustomPrepass = true;
 
 	if ( DetailedObject )
 	{
@@ -774,20 +774,6 @@ void SDataprepDetailsView::AddReferencedObjects(FReferenceCollector& Collector)
 	}
 }
 
-bool SDataprepDetailsView::CustomPrepass(float LayoutScaleMultiplier)
-{
-	if ( bRefreshObjectToDisplay )
-	{
-		TArray< UObject* > Objects;
-		Objects.Add( DetailedObject );
-		Generator->SetObjects( Objects );
-		Construct();
-		bRefreshObjectToDisplay = false;
-	}
-
-	return true;
-}
-
 void SDataprepContextMenuOverride::Construct(const FArguments& InArgs)
 {
 	OnContextMenuOpening = InArgs._OnContextMenuOpening;
diff --git a/Engine/Plugins/Enterprise/DataprepEditor/Source/DataprepEditor/Private/Widgets/DataprepWidgets.h b/Engine/Plugins/Enterprise/DataprepEditor/Source/DataprepEditor/Private/Widgets/DataprepWidgets.h
--- a/Engine/Plugins/Enterprise/DataprepEditor/Source/DataprepEditor/Private/Widgets/DataprepWidgets.h
+++ b/Engine/Plugins/Enterprise/DataprepEditor/Source/DataprepEditor/Private/Widgets/DataprepWidgets.h
@@ -124,9 +124,6 @@ public:
 	static void SetDesiredWidth(int32 InDesiredWidth) { DesiredWidth = InDesiredWidth; }
 
 protected:
-	// ueent_hotfix Hack for 4.24 allow to refresh the ui in between two frame without any flickering
-	virtual bool CustomPrepass(float LayoutScaleMultiplier) override;
-
 	// Begin SWidget overrides.
 	// #ueent_todo: This is temporary until we find a better solution to the splitter issue
 	//				See SConstrainedBox's trick in cpp file
@@ -191,9 +188,6 @@ private:
 	/** Points to the currently used column size data. Can be provided via argument as well. */
 	TSharedPtr< FDetailColumnSizeData > ColumnSizeData;
 
-	/** If there is a new object to display on the next tick */
-	uint8 bRefreshObjectToDisplay : 1;
-
 	/**
 	 * Pointer to a Dataprep asset used by the parameterization system.
 	 * It should be nullptr when the parameterization shouldn't be shown
