Dataprep - Fixing duplicated loctext key entry.

diff --git a/Engine/Plugins/Enterprise/DataprepEditor/Source/DataprepEditor/Private/DataprepGraph/DataprepGraphActionNode.cpp b/Engine/Plugins/Enterprise/DataprepEditor/Source/DataprepEditor/Private/DataprepGraph/DataprepGraphActionNode.cpp
--- a/Engine/Plugins/Enterprise/DataprepEditor/Source/DataprepEditor/Private/DataprepGraph/DataprepGraphActionNode.cpp
+++ b/Engine/Plugins/Enterprise/DataprepEditor/Source/DataprepEditor/Private/DataprepGraph/DataprepGraphActionNode.cpp
@@ -29,7 +29,7 @@ UDataprepGraphActionNode::UDataprepGraphActionNode()
 	, ExecutionOrder(INDEX_NONE)
 {
 	bCanRenameNode = true;
-	ActionTitle = LOCTEXT("DefaultNodeTitle", "New Action").ToString();
+	ActionTitle = LOCTEXT("DefaultActionNodeTitle", "New Action").ToString();
 }
 
 void UDataprepGraphActionNode::Initialize(TWeakObjectPtr<UDataprepAsset> InDataprepAssetPtr, UDataprepActionAsset* InDataprepActionAsset, int32 InExecutionOrder)
@@ -147,7 +147,7 @@ UDataprepGraphActionGroupNode::UDataprepGraphActionGroupNode()
 	: ExecutionOrder(INDEX_NONE)
 {
 	bCanRenameNode = false;
-	NodeTitle = LOCTEXT("DefaultNodeTitle", "Action Group").ToString();
+	NodeTitle = LOCTEXT("DefaultActionGroupNodeTitle", "Action Group").ToString();
 }
 
 void UDataprepGraphActionGroupNode::Initialize(TWeakObjectPtr<UDataprepAsset> InDataprepAssetPtr, TArray<UDataprepActionAsset*>& InActions, int32 InExecutionOrder) 
