Fix cases of engine classes not returning normalized INI paths.

diff --git a/Engine/Plugins/Editor/GameplayTagsEditor/Source/GameplayTagsEditor/Private/SGameplayTagWidget.cpp b/Engine/Plugins/Editor/GameplayTagsEditor/Source/GameplayTagsEditor/Private/SGameplayTagWidget.cpp
--- a/Engine/Plugins/Editor/GameplayTagsEditor/Source/GameplayTagsEditor/Private/SGameplayTagWidget.cpp
+++ b/Engine/Plugins/Editor/GameplayTagsEditor/Source/GameplayTagsEditor/Private/SGameplayTagWidget.cpp
@@ -1085,8 +1085,7 @@ const FString& SGameplayTagWidget::GetGameplayTagsEditorStateIni()
 
 	if (Filename.Len() == 0)
 	{
-		Filename = FString::Printf(TEXT("%s%s/GameplayTagsEditorState.ini"), *FPaths::GeneratedConfigDir(), ANSI_TO_TCHAR(FPlatformProperties::PlatformName()));
-		FPaths::MakeStandardFilename(Filename);
+		Filename = FConfigCacheIni::NormalizeConfigIniPath(FString::Printf(TEXT("%s%s/GameplayTagsEditorState.ini"), *FPaths::GeneratedConfigDir(), ANSI_TO_TCHAR(FPlatformProperties::PlatformName())));
 	}
 
 	return Filename;
diff --git a/Engine/Plugins/Enterprise/DataprepEditor/Source/DataprepEditor/Private/Widgets/SDataprepEditorViewport.cpp b/Engine/Plugins/Enterprise/DataprepEditor/Source/DataprepEditor/Private/Widgets/SDataprepEditorViewport.cpp
--- a/Engine/Plugins/Enterprise/DataprepEditor/Source/DataprepEditor/Private/Widgets/SDataprepEditorViewport.cpp
+++ b/Engine/Plugins/Enterprise/DataprepEditor/Source/DataprepEditor/Private/Widgets/SDataprepEditorViewport.cpp
@@ -1282,7 +1282,7 @@ void SDataprepEditorViewport::LoadDefaultSettings()
 	FPreviewSceneProfile& DataprepViewportSettingProfile = 	DefaultSettings->Profiles[AssetViewerProfileIndex];
 
 	// Read default settings, tessellation and import, for Datasmith file producer
-	const FString DataprepEditorIni = FString::Printf(TEXT("%s%s/%s.ini"), *FPaths::GeneratedConfigDir(), ANSI_TO_TCHAR(FPlatformProperties::PlatformName()), TEXT("DataprepEditor") );
+	const FString DataprepEditorIni = FConfigCacheIni::NormalizeConfigIniPath(FString::Printf(TEXT("%s%s/%s.ini"), *FPaths::GeneratedConfigDir(), ANSI_TO_TCHAR(FPlatformProperties::PlatformName()), TEXT("DataprepEditor")));
 
 	const TCHAR* ViewportSectionName = TEXT("ViewportSettings");
 	if(GConfig->DoesSectionExist( ViewportSectionName, DataprepEditorIni ))
diff --git a/Engine/Plugins/Enterprise/DatasmithImporter/Source/DatasmithImporter/Private/DatasmithFileProducer.cpp b/Engine/Plugins/Enterprise/DatasmithImporter/Source/DatasmithImporter/Private/DatasmithFileProducer.cpp
--- a/Engine/Plugins/Enterprise/DatasmithImporter/Source/DatasmithImporter/Private/DatasmithFileProducer.cpp
+++ b/Engine/Plugins/Enterprise/DatasmithImporter/Source/DatasmithImporter/Private/DatasmithFileProducer.cpp
@@ -1931,7 +1931,7 @@ void FDatasmithFileProducerUtils::DeletePackagePath( const FString& PathToDelete
 void UDatasmithFileProducer::LoadDefaultSettings()
 {
 	// Read default settings, tessellation and import, for Datasmith file producer
-	const FString DatasmithImporterIni = FString::Printf(TEXT("%s%s/%s.ini"), *FPaths::GeneratedConfigDir(), ANSI_TO_TCHAR(FPlatformProperties::PlatformName()), TEXT("DatasmithImporter") );
+	const FString DatasmithImporterIni = FConfigCacheIni::NormalizeConfigIniPath(FString::Printf(TEXT("%s%s/%s.ini"), *FPaths::GeneratedConfigDir(), ANSI_TO_TCHAR(FPlatformProperties::PlatformName()), TEXT("DatasmithImporter")));
 
 	const TCHAR* TessellationSectionName = TEXT("FileProducerTessellationOptions");
 	if(GConfig->DoesSectionExist( TessellationSectionName, DatasmithImporterIni ))
diff --git a/Engine/Source/Runtime/CoreUObject/Private/UObject/Obj.cpp b/Engine/Source/Runtime/CoreUObject/Private/UObject/Obj.cpp
--- a/Engine/Source/Runtime/CoreUObject/Private/UObject/Obj.cpp
+++ b/Engine/Source/Runtime/CoreUObject/Private/UObject/Obj.cpp
@@ -2814,9 +2814,9 @@ FString UObject::GetDefaultConfigFilename() const
 			SelectedPath = FString::Printf(TEXT("%s%s/Config"), *FPaths::ProjectPlatformExtensionsDir(), *OverridePlatform);
 		}
 
-		return FString::Printf(TEXT("%s/%s%s.ini"), *SelectedPath, *OverridePlatform, *GetClass()->ClassConfigName.ToString());
+		return FConfigCacheIni::NormalizeConfigIniPath(FString::Printf(TEXT("%s/%s%s.ini"), *SelectedPath, *OverridePlatform, *GetClass()->ClassConfigName.ToString()));
 	}
-	return FString::Printf(TEXT("%sDefault%s.ini"), *FPaths::SourceConfigDir(), *GetClass()->ClassConfigName.ToString());
+	return FConfigCacheIni::NormalizeConfigIniPath(FString::Printf(TEXT("%sDefault%s.ini"), *FPaths::SourceConfigDir(), *GetClass()->ClassConfigName.ToString()));
 }
 
 FString UObject::GetGlobalUserConfigFilename() const
