Fix spelling error
#rnx

diff --git a/Engine/Plugins/Enterprise/DataprepEditor/Source/DataprepEditor/Private/Widgets/SDataprepEditorViewport.cpp b/Engine/Plugins/Enterprise/DataprepEditor/Source/DataprepEditor/Private/Widgets/SDataprepEditorViewport.cpp
--- a/Engine/Plugins/Enterprise/DataprepEditor/Source/DataprepEditor/Private/Widgets/SDataprepEditorViewport.cpp
+++ b/Engine/Plugins/Enterprise/DataprepEditor/Source/DataprepEditor/Private/Widgets/SDataprepEditorViewport.cpp
@@ -522,7 +522,7 @@ void SDataprepEditorViewport::UpdateScene()
 				{
 					if(UCustomStaticMeshComponent* PreviewMeshComponent = Cast<UCustomStaticMeshComponent>(PreviewMeshComponentPtr.Get()))
 					{
-						PreviewMeshComponent->bShoudlBeInstanced = ShouldBeInstanced.Contains(PreviewMeshComponent->GetStaticMesh());
+						PreviewMeshComponent->bShouldBeInstanced = ShouldBeInstanced.Contains(PreviewMeshComponent->GetStaticMesh());
 						PreviewMeshComponent->MeshProperties = MeshPropertiesMap[PreviewMeshComponent->GetStaticMesh()];
 					}
 				}
@@ -1450,7 +1450,7 @@ void FDataprepEditorViewportClient::Draw(const FSceneView* View, FPrimitiveDrawI
 
 				for(int32 Index = 0; Index < 24; Index += 2)
 				{
-					PDI->DrawLine( Positions[Indices[Index + 0]], Positions[Indices[Index + 1]], MeshComponent->bShoudlBeInstanced ? FColor( 255, 0, 0 ) : FColor( 255, 255, 0 ), SDPG_World );
+					PDI->DrawLine( Positions[Indices[Index + 0]], Positions[Indices[Index + 1]], MeshComponent->bShouldBeInstanced ? FColor( 255, 0, 0 ) : FColor( 255, 255, 0 ), SDPG_World );
 				}
 
 				FVector TransformedCenter = Transform.TransformPosition(Box.Center);
diff --git a/Engine/Plugins/Enterprise/DataprepEditor/Source/DataprepEditor/Private/Widgets/SDataprepEditorViewport.h b/Engine/Plugins/Enterprise/DataprepEditor/Source/DataprepEditor/Private/Widgets/SDataprepEditorViewport.h
--- a/Engine/Plugins/Enterprise/DataprepEditor/Source/DataprepEditor/Private/Widgets/SDataprepEditorViewport.h
+++ b/Engine/Plugins/Enterprise/DataprepEditor/Source/DataprepEditor/Private/Widgets/SDataprepEditorViewport.h
@@ -48,7 +48,7 @@ public:
 
 #ifdef VIEWPORT_EXPERIMENTAL
 	FPrototypeOrientedBox MeshProperties;
-	bool bShoudlBeInstanced;
+	bool bShouldBeInstanced;
 #endif
 
 protected:
