Dataprep editor, clean typo in Commit validation message.

diff --git a/Engine/Plugins/Enterprise/DataprepEditor/Source/DataprepEditor/Private/DataprepEditor.cpp b/Engine/Plugins/Enterprise/DataprepEditor/Source/DataprepEditor/Private/DataprepEditor.cpp
--- a/Engine/Plugins/Enterprise/DataprepEditor/Source/DataprepEditor/Private/DataprepEditor.cpp
+++ b/Engine/Plugins/Enterprise/DataprepEditor/Source/DataprepEditor/Private/DataprepEditor.cpp
@@ -934,7 +934,7 @@ void FDataprepEditor::OnCommitWorld()
 	if( bIsFirstRun && DataprepAssetInterfacePtr->HasActions())
 	{
 		const FText Title( LOCTEXT( "DataprepEditor_ProceedWithCommit", "Proceed with commit" ) );
-		const FText Message( LOCTEXT( "DataprepEditor_ConfirmCommitPipelineNotExecuted", "The action pipeline has not been executed.\nDo you want to proceeed with the commit anyway?" ) );
+		const FText Message( LOCTEXT( "DataprepEditor_ConfirmCommitPipelineNotExecuted", "The action pipeline has not been executed.\nDo you want to proceed with the commit anyway?" ) );
 
 		if( FMessageDialog::Open( EAppMsgType::YesNo, Message, &Title ) != EAppReturnType::Yes )
 		{
@@ -945,7 +945,7 @@ void FDataprepEditor::OnCommitWorld()
 	else if( !bIsFirstRun && bPipelineChanged )
 	{
 		const FText Title( LOCTEXT( "DataprepEditor_ProceedWithCommit", "Proceed with commit" ) );
-		const FText Message( LOCTEXT( "DataprepEditor_ConfirmCommitPipelineChanged", "The action pipeline has changed since last execution.\nDo you want to proceeed with the commit anyway?" ) );
+		const FText Message( LOCTEXT( "DataprepEditor_ConfirmCommitPipelineChanged", "The action pipeline has changed since last execution.\nDo you want to proceed with the commit anyway?" ) );
 
 		if( FMessageDialog::Open( EAppMsgType::YesNo, Message, &Title ) != EAppReturnType::Yes )
 		{
