Fixed use of non-inclusive word

diff --git a/Engine/Plugins/Enterprise/DataprepEditor/Source/DataprepEditor/Private/Widgets/SDataprepEditorViewport.cpp b/Engine/Plugins/Enterprise/DataprepEditor/Source/DataprepEditor/Private/Widgets/SDataprepEditorViewport.cpp
--- a/Engine/Plugins/Enterprise/DataprepEditor/Source/DataprepEditor/Private/Widgets/SDataprepEditorViewport.cpp
+++ b/Engine/Plugins/Enterprise/DataprepEditor/Source/DataprepEditor/Private/Widgets/SDataprepEditorViewport.cpp
@@ -888,7 +888,7 @@ void SDataprepEditorViewport::InitializeDefaultMaterials()
 
 	if(!XRayMaterial.IsValid())
 	{
-		XRayMaterial = TWeakObjectPtr<UMaterial>( CreateMaterialFunc("xray_master") );
+		XRayMaterial = TWeakObjectPtr<UMaterial>( CreateMaterialFunc("xray_material") );
 		check( XRayMaterial.IsValid() );
 
 		DataprepCorePrivateUtils::CompileMaterial(XRayMaterial.Get());
