Fixed Dataprep editor's material assets being flagged as dirty when opening a dataprep asset

diff --git a/Engine/Plugins/Enterprise/DataprepEditor/Source/DataprepEditor/Private/DataprepEditor.cpp b/Engine/Plugins/Enterprise/DataprepEditor/Source/DataprepEditor/Private/DataprepEditor.cpp
--- a/Engine/Plugins/Enterprise/DataprepEditor/Source/DataprepEditor/Private/DataprepEditor.cpp
+++ b/Engine/Plugins/Enterprise/DataprepEditor/Source/DataprepEditor/Private/DataprepEditor.cpp
@@ -1571,7 +1571,13 @@ bool FDataprepEditor::OnRequestClose()
 		const FText Title( LOCTEXT( "DataprepEditor_ProceedWithClose", "Proceed with close") );
 		const FText Message( LOCTEXT( "DataprepEditor_ConfirmClose", "Imported data was not committed! Closing the editor will discard imported data.\nDo you want to close anyway?" ) );
 
-		return ( FMessageDialog::Open(EAppMsgType::YesNo, Message, &Title) == EAppReturnType::Yes );
+		if (FMessageDialog::Open(EAppMsgType::YesNo, Message, &Title) == EAppReturnType::Yes)
+		{
+			ResetBuildWorld();
+			return true;
+		}
+		
+		return false;
 	}
 	return !bIgnoreCloseRequest;
 }
diff --git a/Engine/Plugins/Enterprise/DataprepEditor/Source/DataprepEditor/Private/DataprepEditorModule.cpp b/Engine/Plugins/Enterprise/DataprepEditor/Source/DataprepEditor/Private/DataprepEditorModule.cpp
--- a/Engine/Plugins/Enterprise/DataprepEditor/Source/DataprepEditor/Private/DataprepEditorModule.cpp
+++ b/Engine/Plugins/Enterprise/DataprepEditor/Source/DataprepEditor/Private/DataprepEditorModule.cpp
@@ -88,6 +88,8 @@ public:
 
 		FKismetEditorUtilities::UnregisterAutoBlueprintNodeCreation(this);
 
+		SDataprepEditorViewport::ReleaseDefaultMaterials();
+
 		SDataprepGraphEditor::UnRegisterFactories();
 
 		MenuExtensibilityManager.Reset();
diff --git a/Engine/Plugins/Enterprise/DataprepEditor/Source/DataprepEditor/Private/Widgets/SDataprepEditorViewport.cpp b/Engine/Plugins/Enterprise/DataprepEditor/Source/DataprepEditor/Private/Widgets/SDataprepEditorViewport.cpp
--- a/Engine/Plugins/Enterprise/DataprepEditor/Source/DataprepEditor/Private/Widgets/SDataprepEditorViewport.cpp
+++ b/Engine/Plugins/Enterprise/DataprepEditor/Source/DataprepEditor/Private/Widgets/SDataprepEditorViewport.cpp
@@ -45,13 +45,13 @@
 
 #define LOCTEXT_NAMESPACE "DataprepEditorViewport"
 
-TWeakObjectPtr<UMaterial> SDataprepEditorViewport::PreviewMaterial;
-TWeakObjectPtr<UMaterial> SDataprepEditorViewport::XRayMaterial;
-TWeakObjectPtr<UMaterial> SDataprepEditorViewport::BackFaceMaterial;
-TWeakObjectPtr<UMaterial> SDataprepEditorViewport::PerMeshMaterial;
-TWeakObjectPtr<UMaterial> SDataprepEditorViewport::ReflectionMaterial;
-TWeakObjectPtr<UMaterialInstanceConstant> SDataprepEditorViewport::TransparentMaterial;
-TArray<TWeakObjectPtr<UMaterialInstanceConstant>> SDataprepEditorViewport::PerMeshMaterialInstances;
+TStrongObjectPtr<UMaterial> SDataprepEditorViewport::PreviewMaterial;
+TStrongObjectPtr<UMaterial> SDataprepEditorViewport::XRayMaterial;
+TStrongObjectPtr<UMaterial> SDataprepEditorViewport::BackFaceMaterial;
+TStrongObjectPtr<UMaterial> SDataprepEditorViewport::PerMeshMaterial;
+TStrongObjectPtr<UMaterial> SDataprepEditorViewport::ReflectionMaterial;
+TStrongObjectPtr<UMaterialInstanceConstant> SDataprepEditorViewport::TransparentMaterial;
+TArray<TStrongObjectPtr<UMaterialInstanceConstant>> SDataprepEditorViewport::PerMeshMaterialInstances;
 int32 SDataprepEditorViewport::AssetViewerProfileIndex = INDEX_NONE;
 
 
@@ -852,10 +852,13 @@ void SDataprepEditorViewport::InitializeDefaultMaterials()
 	auto CreateMaterialFunc = [](FString MaterialName) -> UMaterial*
 	{
 		FSoftObjectPath MaterialSoftPath = FSoftObjectPath(FString("/DataprepEditor/") + MaterialName + "." + MaterialName);
-		if(UMaterial* Material = Cast< UMaterial >( MaterialSoftPath.TryLoad() ))
+		if(UMaterial* BaseMaterial = Cast< UMaterial >( MaterialSoftPath.TryLoad() ))
 		{
+			// Temporary solution to fix UE-170536: Duplicate material asset
+			// TODO: Proper solution is to use material instances for the get go. Will be done in future release.
+			UMaterial* Material = DuplicateObject< UMaterial >(BaseMaterial, nullptr);
+			Material->ClearFlags(RF_AllFlags);
 			Material->SetFlags(RF_Transient);
-			Material->ClearFlags(RF_Transactional);
 
 			return Material;
 		}
@@ -866,7 +869,9 @@ void SDataprepEditorViewport::InitializeDefaultMaterials()
 
 	if(!PreviewMaterial.IsValid())
 	{
-		PreviewMaterial = TWeakObjectPtr<UMaterial>( CreateMaterialFunc("PreviewMaterial") );
+		TransparentMaterial.Reset();
+
+		PreviewMaterial = TStrongObjectPtr<UMaterial>( CreateMaterialFunc("PreviewMaterial") );
 		check( PreviewMaterial.IsValid() );
 
 		DataprepCorePrivateUtils::CompileMaterial(PreviewMaterial.Get());
@@ -874,7 +879,7 @@ void SDataprepEditorViewport::InitializeDefaultMaterials()
 
 	if(!TransparentMaterial.IsValid())
 	{
-		TransparentMaterial = TWeakObjectPtr<UMaterialInstanceConstant>( NewObject<UMaterialInstanceConstant>( GetTransientPackage(), NAME_None, EObjectFlags::RF_Transient) );
+		TransparentMaterial = TStrongObjectPtr<UMaterialInstanceConstant>( NewObject<UMaterialInstanceConstant>( GetTransientPackage(), NAME_None, EObjectFlags::RF_Transient) );
 		check( TransparentMaterial.IsValid() );
 
 		TransparentMaterial->Parent = PreviewMaterial.Get();
@@ -890,7 +895,7 @@ void SDataprepEditorViewport::InitializeDefaultMaterials()
 
 	if(!XRayMaterial.IsValid())
 	{
-		XRayMaterial = TWeakObjectPtr<UMaterial>( CreateMaterialFunc("xray_material") );
+		XRayMaterial = TStrongObjectPtr<UMaterial>( CreateMaterialFunc("xray_material") );
 		check( XRayMaterial.IsValid() );
 
 		DataprepCorePrivateUtils::CompileMaterial(XRayMaterial.Get());
@@ -898,14 +903,16 @@ void SDataprepEditorViewport::InitializeDefaultMaterials()
 
 	if(!BackFaceMaterial.IsValid())
 	{
-		BackFaceMaterial = TWeakObjectPtr<UMaterial>( CreateMaterialFunc("BackFaceMaterial") );
+		BackFaceMaterial = TStrongObjectPtr<UMaterial>( CreateMaterialFunc("BackFaceMaterial") );
 
 		DataprepCorePrivateUtils::CompileMaterial(BackFaceMaterial.Get());
 	}
 
 	if(!PerMeshMaterial.IsValid())
 	{
-		PerMeshMaterial = TWeakObjectPtr<UMaterial>( CreateMaterialFunc("PerMeshMaterial") );
+		PerMeshMaterialInstances.Empty();
+
+		PerMeshMaterial = TStrongObjectPtr<UMaterial>( CreateMaterialFunc("PerMeshMaterial") );
 		check( PerMeshMaterial.IsValid() );
 
 		DataprepCorePrivateUtils::CompileMaterial(PerMeshMaterial.Get());
@@ -918,11 +925,11 @@ void SDataprepEditorViewport::InitializeDefaultMaterials()
 
 	for(int32 Index = 0; Index < PerMeshMaterialInstances.Num(); ++Index)
 	{
-		TWeakObjectPtr<UMaterialInstanceConstant>& PerMeshMaterialInstance = PerMeshMaterialInstances[Index];
+		TStrongObjectPtr<UMaterialInstanceConstant>& PerMeshMaterialInstance = PerMeshMaterialInstances[Index];
 
 		if(!PerMeshMaterialInstance.IsValid())
 		{
-			PerMeshMaterialInstance = TWeakObjectPtr<UMaterialInstanceConstant>( NewObject<UMaterialInstanceConstant>( GetTransientPackage(), NAME_None, EObjectFlags::RF_Transient) );
+			PerMeshMaterialInstance = TStrongObjectPtr<UMaterialInstanceConstant>( NewObject<UMaterialInstanceConstant>( GetTransientPackage(), NAME_None, EObjectFlags::RF_Transient) );
 			check( PerMeshMaterialInstance.IsValid() );
 
 			PerMeshMaterialInstance->Parent = PerMeshMaterial.Get();
@@ -937,13 +944,24 @@ void SDataprepEditorViewport::InitializeDefaultMaterials()
 
 	if(!ReflectionMaterial.IsValid())
 	{
-		ReflectionMaterial = TWeakObjectPtr<UMaterial>( CreateMaterialFunc("ReflectionMaterial") );
+		ReflectionMaterial = TStrongObjectPtr<UMaterial>( CreateMaterialFunc("ReflectionMaterial") );
 		check( ReflectionMaterial.IsValid() );
 
 		DataprepCorePrivateUtils::CompileMaterial(ReflectionMaterial.Get());
 	}
 }
 
+void SDataprepEditorViewport::ReleaseDefaultMaterials()
+{
+	PreviewMaterial.Reset();
+	XRayMaterial.Reset();
+	BackFaceMaterial.Reset();
+	ReflectionMaterial.Reset();
+	TransparentMaterial.Reset();
+	PerMeshMaterialInstances.Empty();
+	PerMeshMaterial.Reset();
+}
+
 void SDataprepEditorViewport::SetSelection(UStaticMeshComponent* Component)
 {
 	SelectedPreviewComponents.Empty( 1 );
diff --git a/Engine/Plugins/Enterprise/DataprepEditor/Source/DataprepEditor/Private/Widgets/SDataprepEditorViewport.h b/Engine/Plugins/Enterprise/DataprepEditor/Source/DataprepEditor/Private/Widgets/SDataprepEditorViewport.h
--- a/Engine/Plugins/Enterprise/DataprepEditor/Source/DataprepEditor/Private/Widgets/SDataprepEditorViewport.h
+++ b/Engine/Plugins/Enterprise/DataprepEditor/Source/DataprepEditor/Private/Widgets/SDataprepEditorViewport.h
@@ -186,6 +186,9 @@ public:
 	/** Load the settings affecting the viewport, i.e. environment map */
 	static void LoadDefaultSettings();
 
+	/** Release the materials used for the different rendering options in the viewport */
+	static void ReleaseDefaultMaterials();
+
 	int32 GetDrawCallsAverage() const
 	{
 		return AverageDrawCalls;
@@ -370,23 +373,23 @@ private:
 	bool bWireframeRenderingMode;
 
 	/** Material used to create utility material instances */
-	static TWeakObjectPtr<UMaterial> PreviewMaterial;
+	static TStrongObjectPtr<UMaterial> PreviewMaterial;
 
 	/** Stylized XRay Material */
-	static TWeakObjectPtr<UMaterial> XRayMaterial;
+	static TStrongObjectPtr<UMaterial> XRayMaterial;
 
 	/** Material used to display front facing triangle in green and back facing one in red */
-	static TWeakObjectPtr<UMaterial> BackFaceMaterial;
+	static TStrongObjectPtr<UMaterial> BackFaceMaterial;
 
 	/** Transparent material instance used to display non-selected meshes */
-	static TWeakObjectPtr<UMaterialInstanceConstant> TransparentMaterial;
+	static TStrongObjectPtr<UMaterialInstanceConstant> TransparentMaterial;
 
 	/** Fully reflective material used to display surface discontinuity */
-	static TWeakObjectPtr<UMaterial> ReflectionMaterial;
+	static TStrongObjectPtr<UMaterial> ReflectionMaterial;
 
-	/** Materials	 used to display each mesh component in a different color */
-	static TWeakObjectPtr<UMaterial> PerMeshMaterial;
-	static TArray<TWeakObjectPtr<UMaterialInstanceConstant>> PerMeshMaterialInstances;
+	/** Materials used to display each mesh component in a different color */
+	static TStrongObjectPtr<UMaterial> PerMeshMaterial;
+	static TArray<TStrongObjectPtr<UMaterialInstanceConstant>> PerMeshMaterialInstances;
 
 	TWeakPtr<FDataprepEditor> DataprepEditor;
 
