Refactored places fetching the result of NaniteSettings.bEnabled to instead get the value from UStatic::IsNaniteEnabled() so it correctly includes the force flag (except places using it to determine whether settings should be set).

diff --git a/Engine/Plugins/Enterprise/DataprepEditor/Source/DataprepEditor/Private/DataprepStats.cpp b/Engine/Plugins/Enterprise/DataprepEditor/Source/DataprepEditor/Private/DataprepStats.cpp
--- a/Engine/Plugins/Enterprise/DataprepEditor/Source/DataprepEditor/Private/DataprepStats.cpp
+++ b/Engine/Plugins/Enterprise/DataprepEditor/Source/DataprepEditor/Private/DataprepStats.cpp
@@ -285,7 +285,7 @@ private:
 						Stats.AddCount(FDataprepStats::StatNameVertices, RenderData->LODResources[0].GetNumVertices());
 					}
 
-					if (StaticMesh->NaniteSettings.bEnabled && StaticMesh->HasValidNaniteData())
+					if (StaticMesh->IsNaniteEnabled() && StaticMesh->HasValidNaniteData())
 					{
 						const Nanite::FResources& Resources = StaticMesh->GetRenderData()->NaniteResources;
 						if (Resources.RootData.Num() > 0)
diff --git a/Engine/Plugins/Enterprise/DataprepEditor/Source/DataprepEditor/Private/Widgets/SDataprepEditorViewport.cpp b/Engine/Plugins/Enterprise/DataprepEditor/Source/DataprepEditor/Private/Widgets/SDataprepEditorViewport.cpp
--- a/Engine/Plugins/Enterprise/DataprepEditor/Source/DataprepEditor/Private/Widgets/SDataprepEditorViewport.cpp
+++ b/Engine/Plugins/Enterprise/DataprepEditor/Source/DataprepEditor/Private/Widgets/SDataprepEditorViewport.cpp
@@ -386,7 +386,7 @@ void SDataprepEditorViewport::UpdateScene()
 					const FTransform& ComponentToWorldTransform = SceneMeshComponent->GetComponentTransform();
 					SceneBounds += StaticMesh->GetExtendedBounds().GetBox().TransformBy( ComponentToWorldTransform );
 
-					bCanShowNaniteFallbackMenu |= StaticMesh->NaniteSettings.bEnabled;
+					bCanShowNaniteFallbackMenu |= StaticMesh->IsNaniteEnabled();
 				}
 			}
 
@@ -438,7 +438,7 @@ void SDataprepEditorViewport::UpdateScene()
 						FComponentReregisterContext ReregisterContext( PreviewMeshComponent );
 						PreviewMeshComponent->SetStaticMesh( StaticMesh );
 
-						PreviewMeshComponent->bDisplayNaniteFallbackMesh = bShowNaniteFallbackMenuChecked && StaticMesh->NaniteSettings.bEnabled;
+						PreviewMeshComponent->bDisplayNaniteFallbackMesh = bShowNaniteFallbackMenuChecked && StaticMesh->IsNaniteEnabled();
 
 						FTransform ComponentToWorldTransform = SceneMeshComponent->GetComponentTransform();
 
@@ -616,7 +616,7 @@ void SDataprepEditorViewport::UpdateOverlayText()
 
 			FStaticMeshLODResources& LODResource = StaticMesh->GetRenderData()->LODResources[0];
 
-			if (StaticMesh->NaniteSettings.bEnabled && StaticMesh->HasValidNaniteData())
+			if (StaticMesh->IsNaniteEnabled() && StaticMesh->HasValidNaniteData())
 			{
 				const Nanite::FResources& Resources = StaticMesh->GetRenderData()->NaniteResources;
 				if (Resources.RootData.Num() > 0)
@@ -1327,9 +1327,9 @@ void SDataprepEditorViewport::SetShowNaniteFallback(bool bShow)
 		if (UCustomStaticMeshComponent* PreviewMeshComponent = Cast<UCustomStaticMeshComponent>(PreviewMeshComponentPtr.Get()))
 		{
 			const UStaticMesh* StaticMesh = PreviewMeshComponent->GetStaticMesh();
-			if (StaticMesh && StaticMesh->NaniteSettings.bEnabled)
+			if (StaticMesh && StaticMesh->IsNaniteEnabled())
 			{
-				PreviewMeshComponent->bDisplayNaniteFallbackMesh = bShowNaniteFallbackMenuChecked && StaticMesh->NaniteSettings.bEnabled;
+				PreviewMeshComponent->bDisplayNaniteFallbackMesh = bShowNaniteFallbackMenuChecked;
 				PreviewMeshComponent->MarkRenderStateDirty();
 			}
 		}
diff --git a/Engine/Plugins/Experimental/ChaosEditor/Source/FractureEditor/Private/FractureToolGenerators.cpp b/Engine/Plugins/Experimental/ChaosEditor/Source/FractureEditor/Private/FractureToolGenerators.cpp
--- a/Engine/Plugins/Experimental/ChaosEditor/Source/FractureEditor/Private/FractureToolGenerators.cpp
+++ b/Engine/Plugins/Experimental/ChaosEditor/Source/FractureEditor/Private/FractureToolGenerators.cpp
@@ -622,7 +622,7 @@ AGeometryCollectionActor* UFractureToolGenerateAsset::ConvertActorsToGeometryCol
 				if (ComponentStaticMesh != nullptr)
 				{
 					// If any of the static meshes have Nanite enabled, also enable on the new geometry collection asset for convenience.
-					FracturedGeometryCollection->EnableNanite |= ComponentStaticMesh->NaniteSettings.bEnabled;
+					FracturedGeometryCollection->EnableNanite |= ComponentStaticMesh->IsNaniteEnabled();
 				}
 
 				FTransform ComponentTransform(StaticMeshComponent->GetComponentTransform());
diff --git a/Engine/Plugins/Experimental/MeshLODToolset/Source/MeshLODToolset/Private/Tools/LODManagerTool.cpp b/Engine/Plugins/Experimental/MeshLODToolset/Source/MeshLODToolset/Private/Tools/LODManagerTool.cpp
--- a/Engine/Plugins/Experimental/MeshLODToolset/Source/MeshLODToolset/Private/Tools/LODManagerTool.cpp
+++ b/Engine/Plugins/Experimental/MeshLODToolset/Source/MeshLODToolset/Private/Tools/LODManagerTool.cpp
@@ -221,7 +221,7 @@ void ULODManagerTool::UpdateLODInfo()
 	}
 	bLODInfoValid = true;
 
-	LODInfoProperties->bNaniteEnabled = StaticMesh->NaniteSettings.bEnabled;
+	LODInfoProperties->bNaniteEnabled = StaticMesh->IsNaniteEnabled();
 	LODInfoProperties->PercentTriangles = StaticMesh->NaniteSettings.FallbackPercentTriangles;
 
 	TArray<FStaticMaterial> CurMaterialSet = StaticMesh->GetStaticMaterials();
diff --git a/Engine/Plugins/Experimental/PCG/Source/PCG/Private/Helpers/PCGActorHelpers.cpp b/Engine/Plugins/Experimental/PCG/Source/PCG/Private/Helpers/PCGActorHelpers.cpp
--- a/Engine/Plugins/Experimental/PCG/Source/PCG/Private/Helpers/PCGActorHelpers.cpp
+++ b/Engine/Plugins/Experimental/PCG/Source/PCG/Private/Helpers/PCGActorHelpers.cpp
@@ -61,7 +61,7 @@ UInstancedStaticMeshComponent* UPCGActorHelpers::GetOrCreateISMC(AActor* InTarge
 	
 	// Done as in InstancedStaticMesh.cpp
 #if WITH_EDITOR
-	const bool bMeshHasNaniteData = StaticMesh->NaniteSettings.bEnabled;
+	const bool bMeshHasNaniteData = StaticMesh->IsNaniteEnabled();
 #else
 	const bool bMeshHasNaniteData = StaticMesh->GetRenderData()->NaniteResources.PageStreamingStates.Num() > 0;
 #endif
@@ -192,7 +192,7 @@ UPCGManagedISMComponent* UPCGActorHelpers::GetOrCreateManagedISMC(AActor* InTarg
 
 	// Done as in InstancedStaticMesh.cpp
 #if WITH_EDITOR
-	const bool bMeshHasNaniteData = InMesh->NaniteSettings.bEnabled;
+	const bool bMeshHasNaniteData = InMesh->IsNaniteEnabled();
 #else
 	const bool bMeshHasNaniteData = InMesh->GetRenderData()->NaniteResources.PageStreamingStates.Num() > 0;
 #endif
diff --git a/Engine/Plugins/Importers/USDImporter/Source/USDUtilities/Private/USDGeomMeshConversion.cpp b/Engine/Plugins/Importers/USDImporter/Source/USDUtilities/Private/USDGeomMeshConversion.cpp
--- a/Engine/Plugins/Importers/USDImporter/Source/USDUtilities/Private/USDGeomMeshConversion.cpp
+++ b/Engine/Plugins/Importers/USDImporter/Source/USDUtilities/Private/USDGeomMeshConversion.cpp
@@ -1984,7 +1984,7 @@ bool UnrealToUsd::ConvertStaticMesh( const UStaticMesh* StaticMesh, pxr::UsdPrim
 	// If exporting a Nanite mesh, just use the lowest LOD and write the unrealNanite override attribute,
 	// this way we can guarantee it will have Nanite when it's imported back.
 #if WITH_EDITOR
-	if ( StaticMesh->NaniteSettings.bEnabled )
+	if ( StaticMesh->IsNaniteEnabled())
 	{
 		if ( NumLODs > 1 )
 		{
diff --git a/Engine/Source/Developer/MeshBuilder/Private/StaticMeshBuilder.cpp b/Engine/Source/Developer/MeshBuilder/Private/StaticMeshBuilder.cpp
--- a/Engine/Source/Developer/MeshBuilder/Private/StaticMeshBuilder.cpp
+++ b/Engine/Source/Developer/MeshBuilder/Private/StaticMeshBuilder.cpp
@@ -270,7 +270,7 @@ static bool BuildNanite(
 
 bool FStaticMeshBuilder::Build(FStaticMeshRenderData& StaticMeshRenderData, UStaticMesh* StaticMesh, const FStaticMeshLODGroup& LODGroup, bool bGenerateCoarseMeshStreamingLODs, bool bTargetSupportsNanite)
 {
-	const bool bNaniteBuildEnabled = StaticMesh->NaniteSettings.bEnabled || StaticMesh->IsNaniteForceEnabled();
+	const bool bNaniteBuildEnabled = StaticMesh->IsNaniteEnabled();
 	const bool bHaveHiResSourceModel = StaticMesh->IsHiResMeshDescriptionValid();
 	int32 NumTasks = (bNaniteBuildEnabled && bHaveHiResSourceModel) ? (StaticMesh->GetNumSourceModels() + 1) : (StaticMesh->GetNumSourceModels());
 	FScopedSlowTask SlowTask(NumTasks, NSLOCTEXT("StaticMeshEditor", "StaticMeshBuilderBuild", "Building static mesh render data."));
diff --git a/Engine/Source/Editor/StaticMeshEditor/Private/SStaticMeshEditorViewport.cpp b/Engine/Source/Editor/StaticMeshEditor/Private/SStaticMeshEditorViewport.cpp
--- a/Engine/Source/Editor/StaticMeshEditor/Private/SStaticMeshEditorViewport.cpp
+++ b/Engine/Source/Editor/StaticMeshEditor/Private/SStaticMeshEditorViewport.cpp
@@ -227,7 +227,7 @@ bool SStaticMeshEditorViewport::IsShowNaniteFallbackVisible() const
 {
 	const UStaticMesh* PreviewStaticMesh = PreviewMeshComponent ? ToRawPtr(PreviewMeshComponent->GetStaticMesh()) : nullptr;
 
-	return PreviewStaticMesh && PreviewStaticMesh->NaniteSettings.bEnabled ? true : false;
+	return PreviewStaticMesh && PreviewStaticMesh->IsNaniteEnabled() ? true : false;
 }
 
 void SStaticMeshEditorViewport::UpdatePreviewSocketMeshes()
diff --git a/Engine/Source/Editor/StaticMeshEditor/Private/StaticMeshEditorTools.cpp b/Engine/Source/Editor/StaticMeshEditor/Private/StaticMeshEditorTools.cpp
--- a/Engine/Source/Editor/StaticMeshEditor/Private/StaticMeshEditorTools.cpp
+++ b/Engine/Source/Editor/StaticMeshEditor/Private/StaticMeshEditorTools.cpp
@@ -5146,7 +5146,7 @@ bool FLevelOfDetailSettingsLayout::IsNaniteEnabled() const
 {
 	UStaticMesh* StaticMesh = StaticMeshEditor.GetStaticMesh();
 	check(StaticMesh != nullptr);
-	return StaticMesh->NaniteSettings.bEnabled;
+	return StaticMesh->IsNaniteEnabled();
 }
 
 /////////////////////////////////
@@ -5209,7 +5209,7 @@ void FNaniteSettingsLayout::AddToDetailsPanel(IDetailLayoutBuilder& DetailBuilde
 	auto NaniteCategoryContentLambda = [WeakStaticMesh]()
 	{
 		UStaticMesh* StaticMesh = WeakStaticMesh.Get();
-		if (StaticMesh && StaticMesh->NaniteSettings.bEnabled)
+		if (StaticMesh && StaticMesh->IsNaniteEnabled())
 		{
 			if (!StaticMesh->GetHiResSourceModel().SourceImportFilename.IsEmpty())
 			{
@@ -5222,7 +5222,7 @@ void FNaniteSettingsLayout::AddToDetailsPanel(IDetailLayoutBuilder& DetailBuilde
 	auto NaniteCategoryContentTooltipLambda = [WeakStaticMesh]()
 	{
 		UStaticMesh* StaticMesh = WeakStaticMesh.Get();
-		if (StaticMesh && StaticMesh->NaniteSettings.bEnabled)
+		if (StaticMesh && StaticMesh->IsNaniteEnabled())
 		{
 			if (!StaticMesh->GetHiResSourceModel().SourceImportFilename.IsEmpty())
 			{
diff --git a/Engine/Source/Editor/StaticMeshEditor/Private/StaticMeshEditorViewportClient.cpp b/Engine/Source/Editor/StaticMeshEditor/Private/StaticMeshEditorViewportClient.cpp
--- a/Engine/Source/Editor/StaticMeshEditor/Private/StaticMeshEditorViewportClient.cpp
+++ b/Engine/Source/Editor/StaticMeshEditor/Private/StaticMeshEditorViewportClient.cpp
@@ -792,7 +792,7 @@ void FStaticMeshEditorViewportClient::DrawCanvas(FViewport& InViewport, FSceneVi
 			LOD - 1;
 	}();
 
-	if (StaticMesh->NaniteSettings.bEnabled)
+	if (StaticMesh->IsNaniteEnabled())
 	{
 		TextItems.Emplace(FText::Format(NSLOCTEXT("UnrealEd", "NaniteEnabled", "<TextBlock.ShadowedText>Nanite Enabled</> <TextBlock.ShadowedTextWarning>{0}</>"), StaticMeshComponent->bDisplayNaniteFallbackMesh ? NSLOCTEXT("UnrealEd", "ShowingNaniteFallback", "(Showing Fallback)") : FText::GetEmpty()), false, true);
 
@@ -826,7 +826,7 @@ void FStaticMeshEditorViewportClient::DrawCanvas(FViewport& InViewport, FSceneVi
 		}
 	}
 
-	if (!StaticMesh->NaniteSettings.bEnabled || StaticMeshComponent->bDisplayNaniteFallbackMesh)
+	if (!StaticMesh->IsNaniteEnabled() || StaticMeshComponent->bDisplayNaniteFallbackMesh)
 	{
 		const int32 CurrentMinLODLevel = StaticMesh->GetMinLOD().GetValue();
 		const bool bBelowMinLOD = CurrentLODLevel < CurrentMinLODLevel;
@@ -849,7 +849,7 @@ void FStaticMeshEditorViewportClient::DrawCanvas(FViewport& InViewport, FSceneVi
 	const FText StaticMeshTriangleCount = FText::AsNumber(StaticMeshEditorPtr.Pin()->GetNumTriangles(CurrentLODLevel));
 	const FText StaticMeshVertexCount = FText::AsNumber(StaticMeshEditorPtr.Pin()->GetNumVertices(CurrentLODLevel));
 
-	if (StaticMesh->NaniteSettings.bEnabled)
+	if (StaticMesh->IsNaniteEnabled())
 	{
 		if (StaticMesh->GetRenderData())
 		{
diff --git a/Engine/Source/Editor/UnrealEd/Private/Fbx/FbxStaticMeshImport.cpp b/Engine/Source/Editor/UnrealEd/Private/Fbx/FbxStaticMeshImport.cpp
--- a/Engine/Source/Editor/UnrealEd/Private/Fbx/FbxStaticMeshImport.cpp
+++ b/Engine/Source/Editor/UnrealEd/Private/Fbx/FbxStaticMeshImport.cpp
@@ -1715,7 +1715,7 @@ UStaticMesh* UnFbx::FFbxImporter::ImportStaticMeshAsSingle(UObject* InParent, TA
 
 		//When we do a reimport the nanite static mesh import setting should not affect the asset since its a "reimportrestrict" option
 		//In that case we override the import setting
-		ImportOptions->bBuildNanite = ExistingMesh->NaniteSettings.bEnabled;
+		ImportOptions->bBuildNanite = ExistingMesh->IsNaniteEnabled();
 	}
 	else if (ExistingObject)
 	{
@@ -2026,7 +2026,7 @@ UStaticMesh* UnFbx::FFbxImporter::ImportStaticMeshAsSingle(UObject* InParent, TA
 			ExistingBuildSettings.bGenerateLightmapUVs = SrcModel.BuildSettings.bGenerateLightmapUVs;
 			ExistingBuildSettings.DstLightmapIndex = SrcModel.BuildSettings.DstLightmapIndex;
 			MutableExistMeshDataPtr->ExistingLightMapCoordinateIndex = SrcModel.BuildSettings.DstLightmapIndex;
-			MutableExistMeshDataPtr->ExistingNaniteSettings.bEnabled = StaticMesh->NaniteSettings.bEnabled;
+			MutableExistMeshDataPtr->ExistingNaniteSettings.bEnabled = StaticMesh->IsNaniteEnabled();
 
 			StaticMeshImportUtils::RestoreExistingMeshSettings(ExistMeshDataPtr, InStaticMesh, StaticMesh->LODGroup != NAME_None ? INDEX_NONE : LODIndex);
 		}
diff --git a/Engine/Source/Editor/UnrealEd/Private/FbxMeshUtils.cpp b/Engine/Source/Editor/UnrealEd/Private/FbxMeshUtils.cpp
--- a/Engine/Source/Editor/UnrealEd/Private/FbxMeshUtils.cpp
+++ b/Engine/Source/Editor/UnrealEd/Private/FbxMeshUtils.cpp
@@ -78,7 +78,7 @@ namespace FbxMeshUtils
 			}
 
 			// Set a couple of settings that shouldn't change while importing a lod
-			ImportOptions->bBuildNanite = BaseStaticMesh->NaniteSettings.bEnabled;
+			ImportOptions->bBuildNanite = BaseStaticMesh->IsNaniteEnabled();
 			ImportOptions->StaticMeshLODGroup = BaseStaticMesh->LODGroup;
 			ImportOptions->bIsImportCancelable = false;
 			ImportOptions->bImportMaterials = false;
diff --git a/Engine/Source/Runtime/Engine/Classes/Engine/StaticMesh.h b/Engine/Source/Runtime/Engine/Classes/Engine/StaticMesh.h
--- a/Engine/Source/Runtime/Engine/Classes/Engine/StaticMesh.h
+++ b/Engine/Source/Runtime/Engine/Classes/Engine/StaticMesh.h
@@ -862,11 +862,19 @@ public:
 #endif
 	}
 
+#if WITH_EDITORONLY_DATA
+	/**
+	 * Returns true if this SM should have Nanite built for it.
+	 * This also includes the result of IsNaniteForceEnabled().
+	 */
+	ENGINE_API bool IsNaniteEnabled() const;
+
 	/**
 	 * Returns true if this SM should always have Nanite data built.
 	 * This forces the SM to be Nanite even if the flag in the editor is set to false.
 	 */
 	ENGINE_API bool IsNaniteForceEnabled() const;
+#endif
 
 private:
 	UE_DEPRECATED(5.0, "This must be protected for async build, always use the accessors even internally.")
diff --git a/Engine/Source/Runtime/Engine/Private/Components/StaticMeshComponent.cpp b/Engine/Source/Runtime/Engine/Private/Components/StaticMeshComponent.cpp
--- a/Engine/Source/Runtime/Engine/Private/Components/StaticMeshComponent.cpp
+++ b/Engine/Source/Runtime/Engine/Private/Components/StaticMeshComponent.cpp
@@ -478,7 +478,7 @@ void UStaticMeshComponent::CheckForErrors()
 		OwnerName = Owner->GetName();
 	}
 
-	if (GetStaticMesh() != nullptr && GetStaticMesh()->NaniteSettings.bEnabled != 0)
+	if (GetStaticMesh() != nullptr && GetStaticMesh()->IsNaniteEnabled() != 0)
 	{
 		static const auto CVar = IConsoleManager::Get().FindConsoleVariable(TEXT("r.Shadow.Virtual.Enable"));
 		if (CVar->GetInt() == 0)
@@ -503,7 +503,7 @@ void UStaticMeshComponent::CheckForErrors()
 			for (int32 SectionIndex = 0; SectionIndex < MeshLODData.Sections.Num(); SectionIndex++)
 			{
 				FStaticMeshSection& Element = MeshLODData.Sections[SectionIndex];
-				if (Element.NumTriangles == 0 && !GetStaticMesh()->NaniteSettings.bEnabled)
+				if (Element.NumTriangles == 0 && !GetStaticMesh()->IsNaniteEnabled())
 				{
 					ZeroTriangleElements++;
 				}
diff --git a/Engine/Source/Runtime/Engine/Private/DistanceFieldAtlas.cpp b/Engine/Source/Runtime/Engine/Private/DistanceFieldAtlas.cpp
--- a/Engine/Source/Runtime/Engine/Private/DistanceFieldAtlas.cpp
+++ b/Engine/Source/Runtime/Engine/Private/DistanceFieldAtlas.cpp
@@ -219,7 +219,7 @@ void FDistanceFieldVolumeData::CacheDerivedData(const FString& InStaticMeshDeriv
 		NewTask->MaterialBlendModes = MoveTemp(BuildMaterialData);
 
 		// Nanite overrides source static mesh with a coarse representation. Need to load original data before we build the mesh SDF.
-		if (Mesh->NaniteSettings.bEnabled)
+		if (Mesh->IsNaniteEnabled())
 		{
 			IMeshBuilderModule& MeshBuilderModule = IMeshBuilderModule::GetForPlatform(TargetPlatform);
 			if (!MeshBuilderModule.BuildMeshVertexPositions(Mesh, NewTask->SourceMeshData.TriangleIndices, NewTask->SourceMeshData.VertexPositions))
diff --git a/Engine/Source/Runtime/Engine/Private/InstancedStaticMesh.cpp b/Engine/Source/Runtime/Engine/Private/InstancedStaticMesh.cpp
--- a/Engine/Source/Runtime/Engine/Private/InstancedStaticMesh.cpp
+++ b/Engine/Source/Runtime/Engine/Private/InstancedStaticMesh.cpp
@@ -2417,7 +2417,7 @@ FPrimitiveSceneProxy* UInstancedStaticMeshComponent::CreateSceneProxy()
 			Nanite::AuditMaterials(this, NaniteMaterials);
 		}
 
-		const bool bIsMaskingAllowed = Nanite::IsMaskingAllowedForWorld(GetWorld()) || bForceNaniteForMasked;;
+		const bool bIsMaskingAllowed = Nanite::IsMaskingAllowedForWorld(GetWorld()) || bForceNaniteForMasked;
 		if (bUseNanite && NaniteMaterials.IsValid(bIsMaskingAllowed))
 		{
 			return ::new Nanite::FSceneProxy(NaniteMaterials, this);
@@ -4232,7 +4232,7 @@ static bool ComponentRequestsCPUAccess(UInstancedStaticMeshComponent* InComponen
 			//static const auto CVar = IConsoleManager::Get().FindTConsoleVariableDataInt(TEXT("r.Nanite"));
 
 		#if WITH_EDITOR
-			const bool bHasNaniteData = StaticMesh->NaniteSettings.bEnabled;
+			const bool bHasNaniteData = StaticMesh->IsNaniteEnabled();
 		#else
 			const FStaticMeshRenderData* RenderData = StaticMesh->GetRenderData();
 			const bool bHasNaniteData = RenderData->NaniteResources.PageStreamingStates.Num() > 0;
diff --git a/Engine/Source/Runtime/Engine/Private/MeshCardRepresentation.cpp b/Engine/Source/Runtime/Engine/Private/MeshCardRepresentation.cpp
--- a/Engine/Source/Runtime/Engine/Private/MeshCardRepresentation.cpp
+++ b/Engine/Source/Runtime/Engine/Private/MeshCardRepresentation.cpp
@@ -204,7 +204,7 @@ void FCardRepresentationData::CacheDerivedData(const FString& InDDCKey, const IT
 		{
 			NewTask->SourceMeshData = *OptionalSourceMeshData;
 		}
-		else if (Mesh->NaniteSettings.bEnabled)
+		else if (Mesh->IsNaniteEnabled())
 		{
 			IMeshBuilderModule& MeshBuilderModule = IMeshBuilderModule::GetForPlatform(TargetPlatform);
 			if (!MeshBuilderModule.BuildMeshVertexPositions(Mesh, NewTask->SourceMeshData.TriangleIndices, NewTask->SourceMeshData.VertexPositions))
diff --git a/Engine/Source/Runtime/Engine/Private/PackedLevelActor/PackedLevelActorISMBuilder.cpp b/Engine/Source/Runtime/Engine/Private/PackedLevelActor/PackedLevelActorISMBuilder.cpp
--- a/Engine/Source/Runtime/Engine/Private/PackedLevelActor/PackedLevelActorISMBuilder.cpp
+++ b/Engine/Source/Runtime/Engine/Private/PackedLevelActor/PackedLevelActorISMBuilder.cpp
@@ -49,7 +49,7 @@ void FPackedLevelActorISMBuilder::PackActors(FPackedLevelActorBuilderContext& In
 	ISMComponentBatcher.Append(InComponents, [&ActorTransform, &InContext](const FTransform& InTransform) { return InTransform.GetRelativeTransform(ActorTransform) * InContext.GetRelativePivotTransform() * ActorTransform; });
 
 	TSubclassOf<UInstancedStaticMeshComponent> ComponentClass = UInstancedStaticMeshComponent::StaticClass();
-	if (ISMComponentBatcher.GetNumInstances() > 1 && ISMCluster->ISMDescriptor.StaticMesh && !ISMCluster->ISMDescriptor.StaticMesh->NaniteSettings.bEnabled)
+	if (ISMComponentBatcher.GetNumInstances() > 1 && ISMCluster->ISMDescriptor.StaticMesh && !ISMCluster->ISMDescriptor.StaticMesh->IsNaniteEnabled())
 	{
 		// Use HISM for non-nanite when there is more than one transform (no use in using HISM cpu occlusion for a single instance)
 		ComponentClass = UHierarchicalInstancedStaticMeshComponent::StaticClass();
diff --git a/Engine/Source/Runtime/Engine/Private/StaticMesh.cpp b/Engine/Source/Runtime/Engine/Private/StaticMesh.cpp
--- a/Engine/Source/Runtime/Engine/Private/StaticMesh.cpp
+++ b/Engine/Source/Runtime/Engine/Private/StaticMesh.cpp
@@ -938,7 +938,7 @@ void FStaticMeshVertexFactories::InitVertexFactory(
 	Params.LODIndex						= LODIndex;
 #if WITH_EDITORONLY_DATA
 	Params.StaticMesh					= InParentMesh;
-	Params.bIsCoarseProxy				= InParentMesh->NaniteSettings.bEnabled && InParentMesh->NaniteSettings.FallbackPercentTriangles < 1.0f;
+	Params.bIsCoarseProxy				= InParentMesh->IsNaniteEnabled() && InParentMesh->NaniteSettings.FallbackPercentTriangles < 1.0f;
 #endif
 
 	// Initialize the static mesh's vertex factory.
@@ -6427,7 +6427,7 @@ void UStaticMesh::CheckForMissingShaderModels()
 	// Don't show the SM6 toasts on non-Windows platforms to avoid confusion around platform requirements.
 #if PLATFORM_WINDOWS
 	static bool bWarnedAboutMissingShaderModel = false;
-	if (GIsEditor && NaniteSettings.bEnabled && !bWarnedAboutMissingShaderModel)
+	if (GIsEditor && IsNaniteEnabled() && !bWarnedAboutMissingShaderModel)
 	{
 		TArray<FString> D3D11TargetedShaderFormats;
 		GConfig->GetArray(TEXT("/Script/WindowsTargetPlatform.WindowsTargetSettings"), TEXT("D3D11TargetedShaderFormats"), D3D11TargetedShaderFormats, GEngineIni);
@@ -8266,12 +8266,21 @@ void UStaticMesh::OnLodStrippingQualityLevelChanged(IConsoleVariable* Variable){
 #endif
 }
 
+#if WITH_EDITORONLY_DATA
+
+bool UStaticMesh::IsNaniteEnabled() const
+{
+	return NaniteSettings.bEnabled || IsNaniteForceEnabled();
+}
+
 bool UStaticMesh::IsNaniteForceEnabled() const
 {
 	static const bool bForceEnabled = !!CVarForceEnableNaniteMeshes.GetValueOnAnyThread();
 	return bForceEnabled;
 }
 
+#endif
+
 /*-----------------------------------------------------------------------------
 UStaticMeshSocket
 -----------------------------------------------------------------------------*/
