Dataprep - Replace usage of IAssetTypeActions_AssetInterface with UAssetDefinition

diff --git a/Engine/Plugins/Enterprise/DataprepEditor/Source/DataprepEditor/DataprepEditor.Build.cs b/Engine/Plugins/Enterprise/DataprepEditor/Source/DataprepEditor/DataprepEditor.Build.cs
--- a/Engine/Plugins/Enterprise/DataprepEditor/Source/DataprepEditor/DataprepEditor.Build.cs
+++ b/Engine/Plugins/Enterprise/DataprepEditor/Source/DataprepEditor/DataprepEditor.Build.cs
@@ -12,11 +12,13 @@ namespace UnrealBuildTool.Rules
 					"AdvancedPreviewScene",
 					"ApplicationCore",
 					"AssetRegistry",
+					"AssetDefinition",
 					"AssetTools",
 					"BlueprintGraph",
 					"Blutility",
 					"Core",
 					"CoreUObject",
+					"ContentBrowser",
 					"DataprepCore",
 					"DesktopPlatform",
 					"EditorFramework",
diff --git a/Engine/Plugins/Enterprise/DataprepEditor/Source/DataprepEditor/Private/AssetDefinition_DataprepAsset.cpp b/Engine/Plugins/Enterprise/DataprepEditor/Source/DataprepEditor/Private/AssetDefinition_DataprepAsset.cpp
new file mode 100644
--- /dev/null
+++ b/Engine/Plugins/Enterprise/DataprepEditor/Source/DataprepEditor/Private/AssetDefinition_DataprepAsset.cpp
@@ -0,0 +1,149 @@
+#include "AssetDefinition_DataprepAsset.h"
+
+#include "DataprepAsset.h"
+#include "DataprepAssetInstance.h"
+#include "DataprepCoreUtils.h"
+#include "DataprepEditor.h"
+#include "DataprepEditorModule.h"
+#include "DataprepFactories.h"
+
+#include "ContentBrowserMenuContexts.h"
+
+#include "Algo/AnyOf.h"
+#include "Algo/AllOf.h"
+
+FText UAssetDefinition_DataprepAssetInterface::GetAssetDisplayName() const
+{
+	return NSLOCTEXT("AssetTypeActions_DataprepAssetInterface",  "Name", "Dataprep Interface" );
+}
+
+TSoftClassPtr<> UAssetDefinition_DataprepAssetInterface::GetAssetClass() const
+{
+	return UDataprepAssetInterface::StaticClass();
+}
+
+TConstArrayView<FAssetCategoryPath> UAssetDefinition_DataprepAssetInterface::GetAssetCategories() const
+{
+	static const auto Categories = { FAssetCategoryPath(NSLOCTEXT("AssetTypeActions_DataprepAsset", "Name", "Dataprep Asset")) };
+	return Categories;
+}
+
+EAssetCommandResult UAssetDefinition_DataprepAssetInterface::OpenAssets(const FAssetOpenArgs& OpenArgs) const
+{
+	for (UDataprepAssetInterface* DataprepAssetInterface : OpenArgs.LoadObjects<UDataprepAssetInterface>())
+	{
+		TSharedRef<FDataprepEditor> NewDataprepEditor(new FDataprepEditor());
+		NewDataprepEditor->InitDataprepEditor(EToolkitMode::Standalone, OpenArgs.ToolkitHost, DataprepAssetInterface);
+	}
+
+	return EAssetCommandResult::Handled;
+}
+
+FText UAssetDefinition_DataprepAsset::GetAssetDisplayName() const
+{
+	return NSLOCTEXT("AssetTypeActions_DataprepAsset", "Name", "Dataprep Asset");
+}
+
+TSoftClassPtr<UObject> UAssetDefinition_DataprepAsset::GetAssetClass() const
+{
+	return UDataprepAsset::StaticClass();
+}
+
+FText UAssetDefinition_DataprepAssetInstance::GetAssetDisplayName() const
+{
+	return NSLOCTEXT("AssetTypeActions_DataprepAssetInstance", "Name", "Dataprep Asset Instance");
+}
+
+TSoftClassPtr<UObject> UAssetDefinition_DataprepAssetInstance::GetAssetClass() const
+{
+	return UDataprepAssetInstance::StaticClass();
+}
+
+namespace MenuExtension_DataprepAsset
+{
+	static void CreateInstance(const FToolMenuContext& MenuContext)
+	{
+		// Code is inspired from MenuExtension_MaterialInterface::ExecuteNewMIC
+		const UContentBrowserAssetContextMenuContext* CBContext = UContentBrowserAssetContextMenuContext::FindContextWithAssets(MenuContext);
+
+		IAssetTools::Get().CreateAssetsFrom<UDataprepAssetInterface>(
+			CBContext->LoadSelectedObjects<UDataprepAssetInterface>(), UDataprepAssetInstance::StaticClass(), TEXT("_Inst"), [](UDataprepAssetInterface* SourceObject)
+			{
+				UDataprepAssetInstanceFactory* Factory = nullptr;
+				if(UDataprepAsset* DataprepAsset = Cast<UDataprepAsset>(SourceObject))
+				{
+					Factory = NewObject<UDataprepAssetInstanceFactory>();
+					Factory->Parent = DataprepAsset;
+				}
+				return Factory;
+			}
+		);
+	}
+
+	static void ExecuteDataprepAssets(const FToolMenuContext& MenuContext)
+	{
+		const UContentBrowserAssetContextMenuContext* CBContext = UContentBrowserAssetContextMenuContext::FindContextWithAssets(MenuContext);
+		for (UDataprepAssetInterface* DataprepAssetInterface : CBContext->LoadSelectedObjects<UDataprepAssetInterface>())
+		{
+			// Nothing to do if the Dataprep asset does not have any inputs
+			if(DataprepAssetInterface->GetProducers()->GetProducersCount() > 0)
+			{
+				FDataprepCoreUtils::ExecuteDataprep( DataprepAssetInterface
+					, MakeShared<FDataprepCoreUtils::FDataprepLogger>()
+					, MakeShared<FDataprepCoreUtils::FDataprepProgressUIReporter>() );
+			}
+		}
+	}
+
+	static FDelayedAutoRegisterHelper DelayedAutoRegister(EDelayedRegisterRunPhase::EndOfEngineInit, []{ 
+		UToolMenus::RegisterStartupCallback(FSimpleMulticastDelegate::FDelegate::CreateLambda([]()
+		{
+			FToolMenuOwnerScoped OwnerScoped(UE_MODULE_NAME);
+			UToolMenu* Menu = UE::ContentBrowser::ExtendToolMenu_AssetContextMenu(UDataprepAssetInterface::StaticClass());
+	        
+			FToolMenuSection& Section = Menu->FindOrAddSection("GetAssetActions");
+			Section.AddDynamicEntry(NAME_None, FNewToolMenuSectionDelegate::CreateLambda([](FToolMenuSection& InSection)
+			{
+				if (const UContentBrowserAssetContextMenuContext* CBContext = UContentBrowserAssetContextMenuContext::FindContextWithAssets(InSection))
+				{
+					// #ueent_remark: An instance of an instance is not supported for 4.24.
+					// Do not expose 'Create Instance' menu entry if at least one Dataprep asset is an instance
+					bool bContainsAnInstance  = Algo::AnyOf(CBContext->SelectedAssets, [](const FAssetData& InAssetData)
+					{
+						return InAssetData.GetClass() == UDataprepAssetInstance::StaticClass();
+					});
+
+					// Disable execute if any of the input assets has no consumer
+					bool bCanExecute  = Algo::AllOf(CBContext->SelectedAssets, [](const FAssetData& InAssetData)
+					{
+						if (UDataprepAssetInterface* DataprepAsset = Cast<UDataprepAssetInterface>(InAssetData.GetAsset()))
+						{
+							return DataprepAsset->GetConsumer() != nullptr;
+						}
+						return true;
+					});
+					
+					if (!bContainsAnInstance)
+					{
+						const TAttribute<FText> Label = NSLOCTEXT("AssetTypeActions_DataprepAssetInterface", "CreateInstance", "Create Instance");
+						const TAttribute<FText> ToolTip = NSLOCTEXT("AssetTypeActions_DataprepAssetInterface", "CreateInstanceTooltip", "Creates a parameterized Dataprep asset using this Dataprep asset as a base.");
+						const FSlateIcon Icon = FSlateIcon();
+						const FToolMenuExecuteAction UIAction = FToolMenuExecuteAction::CreateStatic(&CreateInstance);
+
+						InSection.AddMenuEntry("DataprepAssetInterface_CreateInstance", Label, ToolTip, Icon, UIAction);
+					}
+
+					const TAttribute<FText> Label = NSLOCTEXT("AssetTypeActions_DataprepAssetInterface", "RunAsset", "Execute");
+					const TAttribute<FText> ToolTip = NSLOCTEXT("AssetTypeActions_DataprepAssetInterface", "RunAssetTooltip", "Runs the Dataprep asset's producers, execute its recipe, finally runs the consumer");
+					const FSlateIcon Icon = FSlateIcon();
+
+					FToolUIAction UIAction;
+					UIAction.ExecuteAction = FToolMenuExecuteAction::CreateStatic(&ExecuteDataprepAssets);
+					UIAction.CanExecuteAction = FToolMenuCanExecuteAction::CreateLambda([bCanExecute](const FToolMenuContext& Context){ return bCanExecute; });
+
+					InSection.AddMenuEntry("DataprepAssetInterface_Execute", Label, ToolTip, Icon, UIAction);
+				}
+			}));
+		}));
+	});
+}
diff --git a/Engine/Plugins/Enterprise/DataprepEditor/Source/DataprepEditor/Private/AssetDefinition_DataprepAsset.h b/Engine/Plugins/Enterprise/DataprepEditor/Source/DataprepEditor/Private/AssetDefinition_DataprepAsset.h
new file mode 100644
--- /dev/null
+++ b/Engine/Plugins/Enterprise/DataprepEditor/Source/DataprepEditor/Private/AssetDefinition_DataprepAsset.h
@@ -0,0 +1,51 @@
+// Copyright Epic Games, Inc. All Rights Reserved.
+
+#pragma once
+
+#include "AssetDefinitionDefault.h"
+#include "AssetTypeActions_Base.h"
+
+#include "AssetDefinition_DataprepAsset.generated.h"
+
+
+UCLASS()
+class UAssetDefinition_DataprepAssetInterface : public UAssetDefinitionDefault
+{
+	GENERATED_BODY()
+
+public:
+	// UAssetDefinition Begin
+	virtual FText GetAssetDisplayName() const override;
+	virtual FLinearColor GetAssetColor() const override { return FColor(0, 0, 0); }
+	virtual TSoftClassPtr<UObject> GetAssetClass() const override;
+	virtual EAssetCommandResult OpenAssets(const FAssetOpenArgs& OpenArgs) const override;
+	virtual TConstArrayView<FAssetCategoryPath> GetAssetCategories() const override;
+	virtual bool CanImport() const override { return false; }
+	// UAssetDefinition End
+};
+
+UCLASS()
+class UAssetDefinition_DataprepAsset : public UAssetDefinition_DataprepAssetInterface
+{
+	GENERATED_BODY()
+
+public:
+	// UAssetDefinition Begin
+	virtual FText GetAssetDisplayName() const override;
+	virtual FLinearColor GetAssetColor() const override { return FColor(255, 255, 0); }
+	virtual TSoftClassPtr<UObject> GetAssetClass() const override;
+	// UAssetDefinition End
+};
+
+UCLASS()
+class UAssetDefinition_DataprepAssetInstance : public UAssetDefinition_DataprepAssetInterface
+{
+	GENERATED_BODY()
+
+public:
+	// UAssetDefinition Begin
+	virtual FText GetAssetDisplayName() const override;
+	virtual FLinearColor GetAssetColor() const override { return FColor(255, 0, 255); }
+	virtual TSoftClassPtr<UObject> GetAssetClass() const override;
+	// UAssetDefinition End
+};
diff --git a/Engine/Plugins/Enterprise/DataprepEditor/Source/DataprepEditor/Private/AssetTypeActions_DataprepAsset.cpp b/Engine/Plugins/Enterprise/DataprepEditor/Source/DataprepEditor/Private/AssetTypeActions_DataprepAsset.cpp
deleted file mode 100644
diff --git a/Engine/Plugins/Enterprise/DataprepEditor/Source/DataprepEditor/Private/AssetTypeActions_DataprepAsset.h b/Engine/Plugins/Enterprise/DataprepEditor/Source/DataprepEditor/Private/AssetTypeActions_DataprepAsset.h
deleted file mode 100644
diff --git a/Engine/Plugins/Enterprise/DataprepEditor/Source/DataprepEditor/Private/AssetTypeActions_DataprepAssetInterface.cpp b/Engine/Plugins/Enterprise/DataprepEditor/Source/DataprepEditor/Private/AssetTypeActions_DataprepAssetInterface.cpp
deleted file mode 100644
diff --git a/Engine/Plugins/Enterprise/DataprepEditor/Source/DataprepEditor/Private/AssetTypeActions_DataprepAssetInterface.h b/Engine/Plugins/Enterprise/DataprepEditor/Source/DataprepEditor/Private/AssetTypeActions_DataprepAssetInterface.h
deleted file mode 100644
diff --git a/Engine/Plugins/Enterprise/DataprepEditor/Source/DataprepEditor/Private/DataprepEditorModule.cpp b/Engine/Plugins/Enterprise/DataprepEditor/Source/DataprepEditor/Private/DataprepEditorModule.cpp
--- a/Engine/Plugins/Enterprise/DataprepEditor/Source/DataprepEditor/Private/DataprepEditorModule.cpp
+++ b/Engine/Plugins/Enterprise/DataprepEditor/Source/DataprepEditor/Private/DataprepEditorModule.cpp
@@ -2,7 +2,6 @@
 
 #include "DataprepEditorModule.h"
 
-#include "AssetTypeActions_DataprepAsset.h"
 #include "DataprepEditorUtils.h"
 #include "DataprepEditorStyle.h"
 #include "ToolMenu.h"
@@ -21,8 +20,6 @@
 
 const FName DataprepEditorAppIdentifier = FName(TEXT("DataprepEditorApp"));
 
-EAssetTypeCategories::Type IDataprepEditorModule::DataprepCategoryBit;
-
 #define LOCTEXT_NAMESPACE "DataprepEditorModule"
 
 class FDataprepEditorModule : public IDataprepEditorModule
@@ -36,24 +33,6 @@ public:
 		MenuExtensibilityManager = MakeShareable(new FExtensibilityManager);
 		ToolBarExtensibilityManager = MakeShareable(new FExtensibilityManager);
 
-		// Register asset type actions for DataPrepRecipe class
-		IAssetTools& AssetTools = FModuleManager::LoadModuleChecked<FAssetToolsModule>("AssetTools").Get();
-
-		// Register Dataprep category to group together asset type actions related to Dataprep
-		DataprepCategoryBit = AssetTools.RegisterAdvancedAssetCategory(FName(TEXT("Dataprep")), LOCTEXT("DataprepAssetCategory", "Dataprep"));
-
-		TSharedPtr<FAssetTypeActions_DataprepAssetInterface> DataprepAssetInterfaceTypeAction = MakeShared<FAssetTypeActions_DataprepAssetInterface>();
-		AssetTools.RegisterAssetTypeActions(DataprepAssetInterfaceTypeAction.ToSharedRef());
-		AssetTypeActionsArray.Add(DataprepAssetInterfaceTypeAction);
-
-		TSharedPtr<FAssetTypeActions_DataprepAsset> DataprepAssetTypeAction = MakeShareable(new FAssetTypeActions_DataprepAsset);
-		AssetTools.RegisterAssetTypeActions(DataprepAssetTypeAction.ToSharedRef());
-		AssetTypeActionsArray.Add(DataprepAssetTypeAction);
-
-		TSharedPtr<FAssetTypeActions_DataprepAssetInstance> DataprepAssetInstanceTypeAction = MakeShareable(new FAssetTypeActions_DataprepAssetInstance);
-		AssetTools.RegisterAssetTypeActions(DataprepAssetInstanceTypeAction.ToSharedRef());
-		AssetTypeActionsArray.Add(DataprepAssetInstanceTypeAction);
-
 		// Register the details customizer
 		FPropertyEditorModule& PropertyModule = FModuleManager::LoadModuleChecked< FPropertyEditorModule >( TEXT("PropertyEditor") );
 		PropertyModule.RegisterCustomClassLayout( UDataprepAssetProducers::StaticClass()->GetFName(), FOnGetDetailCustomizationInstance::CreateStatic( &FDataprepAssetProducersDetails::MakeDetails ) );
@@ -82,17 +61,6 @@ public:
 		MenuExtensibilityManager.Reset();
 		ToolBarExtensibilityManager.Reset();
 
-		// Unregister asset type actions
-		if (FModuleManager::Get().IsModuleLoaded(TEXT("AssetTools")))
-		{
-			IAssetTools& AssetTools = FModuleManager::LoadModuleChecked<FAssetToolsModule>("AssetTools").Get();
-			for (TSharedPtr<FAssetTypeActions_Base>& AssetTypeActions : AssetTypeActionsArray)
-			{
-				AssetTools.UnregisterAssetTypeActions(AssetTypeActions.ToSharedRef());
-			}
-		}
-		AssetTypeActionsArray.Empty();
-
 		FDataprepEditorStyle::Shutdown();
 
 		// Register the details customizer
@@ -136,7 +104,6 @@ public:
 private:
 	TSharedPtr<FExtensibilityManager> MenuExtensibilityManager;
 	TSharedPtr<FExtensibilityManager> ToolBarExtensibilityManager;
-	TArray<TSharedPtr<FAssetTypeActions_Base>> AssetTypeActionsArray;
 };
 
 IMPLEMENT_MODULE(FDataprepEditorModule, DataprepEditor)
diff --git a/Engine/Plugins/Enterprise/DataprepEditor/Source/DataprepEditor/Public/DataprepEditorModule.h b/Engine/Plugins/Enterprise/DataprepEditor/Source/DataprepEditor/Public/DataprepEditorModule.h
--- a/Engine/Plugins/Enterprise/DataprepEditor/Source/DataprepEditor/Public/DataprepEditorModule.h
+++ b/Engine/Plugins/Enterprise/DataprepEditor/Source/DataprepEditor/Public/DataprepEditorModule.h
@@ -46,5 +46,4 @@ public:
 	virtual TSharedRef<SWidget> CreateDataprepProducersWidget(UDataprepAssetProducers* AssetProducers) = 0;
 	virtual TSharedRef<SWidget> CreateDataprepDetailsView(UObject* ObjectToDetail) = 0;
 
-	static EAssetTypeCategories::Type DataprepCategoryBit;
 };
