Fixed copyright boilerplate

diff --git a/Engine/Plugins/Enterprise/DataprepEditor/Source/DataprepEditor/Private/AssetDefinition_DataprepAsset.cpp b/Engine/Plugins/Enterprise/DataprepEditor/Source/DataprepEditor/Private/AssetDefinition_DataprepAsset.cpp
--- a/Engine/Plugins/Enterprise/DataprepEditor/Source/DataprepEditor/Private/AssetDefinition_DataprepAsset.cpp
+++ b/Engine/Plugins/Enterprise/DataprepEditor/Source/DataprepEditor/Private/AssetDefinition_DataprepAsset.cpp
@@ -1,3 +1,5 @@
+// Copyright Epic Games, Inc. All Rights Reserved.
+
 #include "AssetDefinition_DataprepAsset.h"
 
 #include "DataprepAsset.h"
