[Content Pipeline] Remove redundant data from telemetry

- USD/AxF/MDL/Datasmith importers, Dataprep

diff --git a/Engine/Plugins/Enterprise/AxFImporter/Source/AxFImporter/Private/AxFImporterFactory.cpp b/Engine/Plugins/Enterprise/AxFImporter/Source/AxFImporter/Private/AxFImporterFactory.cpp
--- a/Engine/Plugins/Enterprise/AxFImporter/Source/AxFImporter/Private/AxFImporterFactory.cpp
+++ b/Engine/Plugins/Enterprise/AxFImporter/Source/AxFImporter/Private/AxFImporterFactory.cpp
@@ -208,7 +208,6 @@ void UAxFImporterFactory::SendAnalytics(int32 ImportDurationInSeconds, bool bImp
 		FString ImporterName = TEXT("AxF");
 
 		EventAttributes.Emplace(TEXT("ImporterType"), ImporterName);
-		EventAttributes.Emplace(TEXT("ImporterID"), FPlatformMisc::GetEpicAccountId());
 		EventAttributes.Emplace(TEXT("ImportDuration"), ImportDurationInSeconds);
 
 		if (bImportSuccess)
diff --git a/Engine/Plugins/Enterprise/DataprepEditor/Source/DataprepCore/Private/DataprepCorePrivateUtils.cpp b/Engine/Plugins/Enterprise/DataprepEditor/Source/DataprepCore/Private/DataprepCorePrivateUtils.cpp
--- a/Engine/Plugins/Enterprise/DataprepEditor/Source/DataprepCore/Private/DataprepCorePrivateUtils.cpp
+++ b/Engine/Plugins/Enterprise/DataprepEditor/Source/DataprepCore/Private/DataprepCorePrivateUtils.cpp
@@ -229,7 +229,6 @@ void DataprepCorePrivateUtils::Analytics::RecipeExecuted( UDataprepAssetInterfac
 	{
 		TArray<FAnalyticsEventAttribute> EventAttributes;
 
-		EventAttributes.Emplace( TEXT("EpicAccountID"), FPlatformMisc::GetEpicAccountId() );
 		EventAttributes.Emplace( TEXT("EngineVersion"), FEngineVersion::Current().ToString( EVersionComponent::Patch ) );
 		
 		const TArray<UDataprepActionAsset*>& Actions = InDataprepAsset->GetActions();
@@ -256,7 +255,6 @@ void DataprepCorePrivateUtils::Analytics::DataprepAssetCreated( UDataprepAssetIn
 	{
 		TArray<FAnalyticsEventAttribute> EventAttributes;
 
-		EventAttributes.Emplace( TEXT("EpicAccountID"), FPlatformMisc::GetEpicAccountId() );
 		EventAttributes.Emplace( TEXT("EngineVersion"), FEngineVersion::Current().ToString( EVersionComponent::Patch ) );
 
 		const bool bIsDataprepInstance = InDataprepAsset->IsA<UDataprepAssetInstance>();
@@ -272,7 +270,6 @@ void DataprepCorePrivateUtils::Analytics::DataprepEditorOpened( UDataprepAssetIn
 	{
 		TArray<FAnalyticsEventAttribute> EventAttributes;
 
-		EventAttributes.Emplace( TEXT("EpicAccountID"), FPlatformMisc::GetEpicAccountId() );
 		EventAttributes.Emplace( TEXT("EngineVersion"), FEngineVersion::Current().ToString( EVersionComponent::Patch ) );
 
 		FEngineAnalytics::GetProvider().RecordEvent( TEXT("Editor.Dataprep.EditorOpened"), EventAttributes );
@@ -285,7 +282,6 @@ void DataprepCorePrivateUtils::Analytics::ExecuteTriggered( UDataprepAssetInterf
 	{
 		TArray<FAnalyticsEventAttribute> EventAttributes;
 
-		EventAttributes.Emplace( TEXT("EpicAccountID"), FPlatformMisc::GetEpicAccountId() );
 		EventAttributes.Emplace( TEXT("EngineVersion"), FEngineVersion::Current().ToString( EVersionComponent::Patch ) );
 
 		FEngineAnalytics::GetProvider().RecordEvent( TEXT("Editor.Dataprep.ExecuteTriggered"), EventAttributes );
@@ -298,7 +294,6 @@ void DataprepCorePrivateUtils::Analytics::ImportTriggered( UDataprepAssetInterfa
 	{
 		TArray<FAnalyticsEventAttribute> EventAttributes;
 
-		EventAttributes.Emplace( TEXT("EpicAccountID"), FPlatformMisc::GetEpicAccountId() );
 		EventAttributes.Emplace( TEXT("EngineVersion"), FEngineVersion::Current().ToString( EVersionComponent::Patch ) );
 
 		FEngineAnalytics::GetProvider().RecordEvent( TEXT("Editor.Dataprep.ImportTriggered"), EventAttributes );
@@ -311,7 +306,6 @@ void DataprepCorePrivateUtils::Analytics::CommitTriggered( UDataprepAssetInterfa
 	{
 		TArray<FAnalyticsEventAttribute> EventAttributes;
 
-		EventAttributes.Emplace( TEXT("EpicAccountID"), FPlatformMisc::GetEpicAccountId() );
 		EventAttributes.Emplace( TEXT("EngineVersion"), FEngineVersion::Current().ToString( EVersionComponent::Patch ) );
 
 		FEngineAnalytics::GetProvider().RecordEvent( TEXT("Editor.Dataprep.CommitTriggered"), EventAttributes );
diff --git a/Engine/Plugins/Enterprise/DatasmithImporter/Source/DatasmithImporter/Private/DatasmithImportFactory.cpp b/Engine/Plugins/Enterprise/DatasmithImporter/Source/DatasmithImporter/Private/DatasmithImportFactory.cpp
--- a/Engine/Plugins/Enterprise/DatasmithImporter/Source/DatasmithImporter/Private/DatasmithImportFactory.cpp
+++ b/Engine/Plugins/Enterprise/DatasmithImporter/Source/DatasmithImporter/Private/DatasmithImportFactory.cpp
@@ -323,7 +323,6 @@ namespace DatasmithImportFactoryImpl
 			EventAttributes.Emplace( TEXT("ImportedWithViaScript"), ImportContext.bImportedViaScript );
 
 			// Import info
-			EventAttributes.Emplace( TEXT("ImporterID"), FPlatformMisc::GetEpicAccountId() );
 			EventAttributes.Emplace( TEXT("ImporterVersion"), FDatasmithUtils::GetEnterpriseVersionAsString() );
 			EventAttributes.Emplace( TEXT("ImportDuration"), ImportDurationInSeconds );
 
diff --git a/Engine/Plugins/Enterprise/MDLImporter/Source/MDLImporter/Private/MDLImporterFactory.cpp b/Engine/Plugins/Enterprise/MDLImporter/Source/MDLImporter/Private/MDLImporterFactory.cpp
--- a/Engine/Plugins/Enterprise/MDLImporter/Source/MDLImporter/Private/MDLImporterFactory.cpp
+++ b/Engine/Plugins/Enterprise/MDLImporter/Source/MDLImporter/Private/MDLImporterFactory.cpp
@@ -211,7 +211,6 @@ void UMDLImporterFactory::SendAnalytics(int32 ImportDurationInSeconds, bool bImp
 		FString ImporterName = TEXT("MDL");
 
 		EventAttributes.Emplace(TEXT("ImporterType"), ImporterName);
-		EventAttributes.Emplace(TEXT("ImporterID"), FPlatformMisc::GetEpicAccountId());
 		EventAttributes.Emplace(TEXT("ImportDuration"), ImportDurationInSeconds);
 
 		if (bImportSuccess)
diff --git a/Engine/Plugins/Importers/USDImporter/Source/USDClasses/Private/USDClassesModule.cpp b/Engine/Plugins/Importers/USDImporter/Source/USDClasses/Private/USDClassesModule.cpp
--- a/Engine/Plugins/Importers/USDImporter/Source/USDClasses/Private/USDClassesModule.cpp
+++ b/Engine/Plugins/Importers/USDImporter/Source/USDClasses/Private/USDClassesModule.cpp
@@ -205,7 +205,6 @@ void IUsdClassesModule::SendAnalytics( TArray<FAnalyticsEventAttribute>&& InAttr
 		Attributes.Emplace( TEXT( "Platform" ), FPlatformProperties::IniPlatformName() );
 		Attributes.Emplace( TEXT( "EngineVersion" ), FEngineVersion::Current().ToString() );
 		Attributes.Emplace( TEXT( "EngineMode" ), FPlatformMisc::GetEngineMode() );
-		Attributes.Emplace( TEXT( "EpicAccountId" ), FPlatformMisc::GetEpicAccountId() );
 
 		const FString EventText = FString::Printf( TEXT( "Engine.Usage.USD.%s" ), *EventName);
 		FEngineAnalytics::GetProvider().RecordEvent( EventText, Attributes );
