[Dataprep] Fixed crash when creating a parameter from a datatable asset or any asset

diff --git a/Engine/Plugins/Enterprise/DataprepEditor/Source/DataprepCore/Private/DataprepParameterizationArchive.cpp b/Engine/Plugins/Enterprise/DataprepEditor/Source/DataprepCore/Private/DataprepParameterizationArchive.cpp
--- a/Engine/Plugins/Enterprise/DataprepEditor/Source/DataprepCore/Private/DataprepParameterizationArchive.cpp
+++ b/Engine/Plugins/Enterprise/DataprepEditor/Source/DataprepCore/Private/DataprepParameterizationArchive.cpp
@@ -41,6 +41,12 @@ FArchive& FDataprepParameterizationWriter::operator<<(FLazyObjectPtr& Value)
 	return *this << Object;
 }
 
+FArchive& FDataprepParameterizationWriter::operator<<(FObjectPtr& Value)
+{
+	UObject* Object = Value.Get();
+	return *this << Object;
+}
+
 FArchive& FDataprepParameterizationWriter::operator<<(FSoftObjectPtr& Value)
 {
 	FSoftObjectPath SoftObjectPath = Value.ToSoftObjectPath();
@@ -118,6 +124,14 @@ FArchive& FDataprepParameterizationReader::operator<<(FLazyObjectPtr& Value)
 	return *this;
 }
 
+FArchive& FDataprepParameterizationReader::operator<<(FObjectPtr& Value)
+{
+	UObject* Object = nullptr;
+	*this << Object;
+	Value = Object;
+	return *this;
+}
+
 FArchive& FDataprepParameterizationReader::operator<<(FSoftObjectPtr& Value)
 {
 	FSoftObjectPath SoftObjectPath;
diff --git a/Engine/Plugins/Enterprise/DataprepEditor/Source/DataprepCore/Private/DataprepParameterizationArchive.h b/Engine/Plugins/Enterprise/DataprepEditor/Source/DataprepCore/Private/DataprepParameterizationArchive.h
--- a/Engine/Plugins/Enterprise/DataprepEditor/Source/DataprepCore/Private/DataprepParameterizationArchive.h
+++ b/Engine/Plugins/Enterprise/DataprepEditor/Source/DataprepCore/Private/DataprepParameterizationArchive.h
@@ -47,6 +47,8 @@ public:
 
 	virtual FArchive& operator<<(struct FLazyObjectPtr& Value) override;
 
+	virtual FArchive& operator<<(struct FObjectPtr& Value) override;
+
 	virtual FArchive& operator<<(struct FSoftObjectPtr& Value) override;
 
 	virtual FArchive& operator<<(struct FSoftObjectPath& Value) override;
@@ -97,6 +99,8 @@ public:
 
 	virtual FArchive& operator<<(struct FLazyObjectPtr& Value) override;
 
+	virtual FArchive& operator<<(struct FObjectPtr& Value) override;
+
 	virtual FArchive& operator<<(struct FSoftObjectPtr& Value) override;
 
 	virtual FArchive& operator<<(struct FSoftObjectPath& Value) override;
