Dataprep crash on multiple execute of graph assigning materials randomly

diff --git a/Engine/Plugins/Enterprise/DataprepEditor/Source/DataprepCore/Private/DataprepCoreUtils.cpp b/Engine/Plugins/Enterprise/DataprepEditor/Source/DataprepCore/Private/DataprepCoreUtils.cpp
--- a/Engine/Plugins/Enterprise/DataprepEditor/Source/DataprepCore/Private/DataprepCoreUtils.cpp
+++ b/Engine/Plugins/Enterprise/DataprepEditor/Source/DataprepCore/Private/DataprepCoreUtils.cpp
@@ -201,6 +201,8 @@ void FDataprepCoreUtils::PurgeObjects(TArray<UObject*> InObjects)
 	// if we have object to purge but the map isn't one of them collect garbage (if we purged the map it has already been done)
 	if ( ObjectsToPurge.Num() > 0 )
 	{
+		FlushRenderingCommands();
+
 		CollectGarbage(GARBAGE_COLLECTION_KEEPFLAGS);
 	}
 }
diff --git a/Engine/Plugins/Enterprise/DataprepEditor/Source/DataprepEditor/Private/DataprepEditor.cpp b/Engine/Plugins/Enterprise/DataprepEditor/Source/DataprepEditor/Private/DataprepEditor.cpp
--- a/Engine/Plugins/Enterprise/DataprepEditor/Source/DataprepEditor/Private/DataprepEditor.cpp
+++ b/Engine/Plugins/Enterprise/DataprepEditor/Source/DataprepEditor/Private/DataprepEditor.cpp
@@ -38,6 +38,8 @@
 #include "Widgets/SDataprepProducersWidget.h"
 #include "Widgets/SDataprepStats.h"
 
+#include "Materials/MaterialInstanceConstant.h"
+
 #define LOCTEXT_NAMESPACE "DataprepEditor"
 
 extern const FName DataprepEditorAppIdentifier;
@@ -791,6 +793,12 @@ void FDataprepEditor::CleanPreviewWorld()
 
 			UObject* ObjectToDelete = StaticFindObjectFast(nullptr, Package, *SoftObjectPath.GetAssetName());
 
+			if (UMaterialInstanceConstant* MaterialInstanceConstant = Cast<UMaterialInstanceConstant>(ObjectToDelete))
+			{
+				MaterialInstanceConstant->SetParentEditorOnly(nullptr);
+				MaterialInstanceConstant->PostEditChange();
+			}
+
 			FDataprepCoreUtils::MoveToTransientPackage( ObjectToDelete );
 			ObjectsToDelete.Add( ObjectToDelete );
 		}
diff --git a/Engine/Plugins/Enterprise/DataprepEditor/Source/DataprepEditor/Private/Widgets/SDataprepEditorViewport.cpp b/Engine/Plugins/Enterprise/DataprepEditor/Source/DataprepEditor/Private/Widgets/SDataprepEditorViewport.cpp
--- a/Engine/Plugins/Enterprise/DataprepEditor/Source/DataprepEditor/Private/Widgets/SDataprepEditorViewport.cpp
+++ b/Engine/Plugins/Enterprise/DataprepEditor/Source/DataprepEditor/Private/Widgets/SDataprepEditorViewport.cpp
@@ -290,18 +290,30 @@ void SDataprepEditorViewport::Construct(const FArguments& InArgs, TSharedPtr<FDa
 void SDataprepEditorViewport::ClearScene()
 {
 	const int32 PreviousCount = PreviewMeshComponents.Num();
+	
+	FDetachmentTransformRules DetachmentRules(EDetachmentRule::KeepRelative, true);
+
+	StaticActor->UnregisterAllComponents();
+	MovableActor->UnregisterAllComponents();
+	
 	TArray<UObject*> ObjectsToDelete;
-	ObjectsToDelete.Reserve( PreviousCount );
+	ObjectsToDelete.Reserve(PreviousCount);
 
 	for( const TWeakObjectPtr< UStaticMeshComponent >& PreviewMeshComponent : PreviewMeshComponents )
 	{
 		if(UStaticMeshComponent* MeshComponent = PreviewMeshComponent.Get())
 		{
+			MeshComponent->DetachFromComponent(DetachmentRules);
+			MeshComponent->SetStaticMesh(nullptr);
+
 			MeshComponent->EmptyOverrideMaterials();
 			ObjectsToDelete.Add( MeshComponent );
 		}
 	}
 
+	StaticActor->RegisterAllComponents();
+	MovableActor->RegisterAllComponents();
+
 	FDataprepCoreUtils::PurgeObjects( MoveTemp( ObjectsToDelete ) );
 
 	PreviewMeshComponents.Empty(PreviousCount);
