Fix Cast used as const_cast
#rnx

diff --git a/Engine/Plugins/Enterprise/DataprepEditor/Source/DataprepEditor/Private/Widgets/SDataprepEditorViewport.cpp b/Engine/Plugins/Enterprise/DataprepEditor/Source/DataprepEditor/Private/Widgets/SDataprepEditorViewport.cpp
--- a/Engine/Plugins/Enterprise/DataprepEditor/Source/DataprepEditor/Private/Widgets/SDataprepEditorViewport.cpp
+++ b/Engine/Plugins/Enterprise/DataprepEditor/Source/DataprepEditor/Private/Widgets/SDataprepEditorViewport.cpp
@@ -1391,7 +1391,7 @@ void FDataprepEditorViewportClient::ProcessClick(FSceneView& View, HHitProxy* Hi
 			if( HitProxy->IsA( HActor::StaticGetType() ) )
 			{
 				// A static mesh component has been selected
-				if( UStaticMeshComponent* Component = Cast<UStaticMeshComponent>( ((HActor*)HitProxy)->PrimComponent ))
+				if( UStaticMeshComponent* Component = const_cast<UStaticMeshComponent*>(Cast<UStaticMeshComponent>( ((HActor*)HitProxy)->PrimComponent )))
 				{
 					// A static mesh component part of the ones to preview has been selected
 					if( DataprepEditorViewport->IsAPreviewComponent(Component) )
diff --git a/Engine/Plugins/Experimental/PCGExternalDataInterop/Source/PCGExternalDataInteropEditor/Private/PCGLoadAlembic.cpp b/Engine/Plugins/Experimental/PCGExternalDataInterop/Source/PCGExternalDataInteropEditor/Private/PCGLoadAlembic.cpp
--- a/Engine/Plugins/Experimental/PCGExternalDataInterop/Source/PCGExternalDataInteropEditor/Private/PCGLoadAlembic.cpp
+++ b/Engine/Plugins/Experimental/PCGExternalDataInterop/Source/PCGExternalDataInteropEditor/Private/PCGLoadAlembic.cpp
@@ -35,7 +35,7 @@ void UPCGLoadAlembicFunctionLibrary::LoadAlembicFileToPCG(const FPCGLoadAlembicB
 		{
 			if (TaggedData.Data)
 			{
-				Cast<UObject>(TaggedData.Data)->Rename(nullptr, TargetOuter, REN_ForceNoResetLoaders | REN_DoNotDirty | REN_DontCreateRedirectors);
+				(const_cast<UPCGData*>(TaggedData.Data.Get()))->Rename(nullptr, TargetOuter, REN_ForceNoResetLoaders | REN_DoNotDirty | REN_DontCreateRedirectors);
 			}
 		}
 	}
