Added Interchange Import to Dataprep - WIP
Set CVar 'Dataprep.Interchange.Producer' to true to use the Interchange File Producer. Default to false.

#rnx

diff --git a/Engine/Plugins/Enterprise/DataprepEditor/DataprepEditor.uplugin b/Engine/Plugins/Enterprise/DataprepEditor/DataprepEditor.uplugin
--- a/Engine/Plugins/Enterprise/DataprepEditor/DataprepEditor.uplugin
+++ b/Engine/Plugins/Enterprise/DataprepEditor/DataprepEditor.uplugin
@@ -48,6 +48,10 @@
 		{
 			"Name": "VariantManagerContent",
 			"Enabled": true
+		},
+		{
+			"Name": "Interchange",
+			"Enabled": true
 		}
 	],
 	"IsExperimentalVersion": false
diff --git a/Engine/Plugins/Enterprise/DataprepEditor/Source/DataprepCore/DataprepCore.Build.cs b/Engine/Plugins/Enterprise/DataprepEditor/Source/DataprepCore/DataprepCore.Build.cs
--- a/Engine/Plugins/Enterprise/DataprepEditor/Source/DataprepCore/DataprepCore.Build.cs
+++ b/Engine/Plugins/Enterprise/DataprepEditor/Source/DataprepCore/DataprepCore.Build.cs
@@ -13,12 +13,17 @@ namespace UnrealBuildTool.Rules
 					"BlueprintGraph",
 					"Core",
 					"CoreUObject",
+					"DesktopPlatform",
 					"EditorFramework",
 					"EditorScriptingUtilities",
+					"EditorStyle",
 					"EditorWidgets",
 					"Engine",
 					"GraphEditor",
 					"InputCore",
+					"InterchangeCore",
+					"InterchangeEngine",
+					"InterchangeFactoryNodes",
 					"KismetCompiler",
 					"LevelSequence",
 					"MeshDescription",
diff --git a/Engine/Plugins/Enterprise/DataprepEditor/Source/DataprepCore/Private/DataprepContentProducer.cpp b/Engine/Plugins/Enterprise/DataprepEditor/Source/DataprepCore/Private/DataprepContentProducer.cpp
--- a/Engine/Plugins/Enterprise/DataprepEditor/Source/DataprepCore/Private/DataprepContentProducer.cpp
+++ b/Engine/Plugins/Enterprise/DataprepEditor/Source/DataprepCore/Private/DataprepContentProducer.cpp
@@ -87,18 +87,21 @@ bool UDataprepContentProducer::Produce(const FDataprepProducerContext& InContext
 	// Prefix all newly created actors with the namespace of the producer
 	const FString Namespace = GetNamespace();
 
-	for (TActorIterator<AActor> It(Context.WorldPtr.Get(), AActor::StaticClass(), Flags); It; ++It)
+	if (Namespace.Len() > 0)
 	{
-		if(*It != nullptr && ExistingActors.Find( *It ) == nullptr)
+		for (TActorIterator<AActor> It(Context.WorldPtr.Get(), AActor::StaticClass(), Flags); It; ++It)
 		{
-			AActor* Actor = *It;
+			if (*It != nullptr && ExistingActors.Find(*It) == nullptr)
+			{
+				AActor* Actor = *It;
 
-			FSoftObjectPath PreviousActorSoftPath(Actor);
+				FSoftObjectPath PreviousActorSoftPath(Actor);
 
-			const FString ActorName =  Namespace + TEXT("_") + Actor->GetName();
-			FDataprepCoreUtils::RenameObject( Actor, *ActorName );
+				const FString ActorName = Namespace + TEXT("_") + Actor->GetName();
+				FDataprepCoreUtils::RenameObject(Actor, *ActorName);
 
-			ActorRedirectorMap.Emplace( PreviousActorSoftPath, Actor );
+				ActorRedirectorMap.Emplace(PreviousActorSoftPath, Actor);
+			}
 		}
 	}
 
diff --git a/Engine/Plugins/Enterprise/DataprepEditor/Source/DataprepCore/Private/InterchangeDataprepPipeline.cpp b/Engine/Plugins/Enterprise/DataprepEditor/Source/DataprepCore/Private/InterchangeDataprepPipeline.cpp
new file mode 100644
--- /dev/null
+++ b/Engine/Plugins/Enterprise/DataprepEditor/Source/DataprepCore/Private/InterchangeDataprepPipeline.cpp
@@ -0,0 +1,68 @@
+// Copyright Epic Games, Inc. All Rights Reserved.
+
+#include "InterchangeDataprepPipeline.h"
+
+//#include "Nodes/InterchangeFactoryBaseNode.h"
+#include "InterchangeActorFactoryNode.h"
+
+#include "Engine/Texture.h"
+#include "Materials/MaterialInterface.h"
+#include "Engine/StaticMesh.h"
+#include "LevelSequence.h"
+
+void UInterchangeDataprepLevelPipeline::ExecutePipeline(UInterchangeBaseNodeContainer* NodeContainer, const TArray<UInterchangeSourceData*>& InSourceDatas)
+{
+	if (!NodeContainer)
+	{
+		return;
+	}
+
+	Super::ExecutePipeline(NodeContainer, InSourceDatas);
+
+	static TMap<UClass*, FString> SubPathsPerClass{
+		{ UTexture::StaticClass(), "Textures" },
+		{ UMaterialInterface::StaticClass(), "Materials" },
+		{ UStaticMesh::StaticClass(), "Geometries" },
+		{ ULevelSequence::StaticClass(), TEXT("Animations") },
+	};
+
+	// Compute unique prefix based on file path
+	FString UniquePrefix = FString::FromInt(GetTypeHash(InSourceDatas[0]->GetFilename()));
+
+	auto UpdateFactoryNodes = [this,&UniquePrefix](const FString& NodeUid, UInterchangeFactoryBaseNode* FactoryNode)
+	{
+		for (TPair<UClass*, FString>& Entry : SubPathsPerClass)
+		{
+			if (FactoryNode->GetObjectClass() && FactoryNode->GetObjectClass()->IsChildOf(Entry.Key))
+			{
+				FactoryNode->SetCustomSubPath(Entry.Value);
+				break;
+			}
+		}
+
+		// Append prefix to prevent actor's name collision in Dataprep world
+		if (UInterchangeActorFactoryNode* ActorFactoryNode = Cast< UInterchangeActorFactoryNode>(FactoryNode))
+		{
+			FString NewLabel = UniquePrefix + TEXT("_") + ActorFactoryNode->GetDisplayLabel();
+			ActorFactoryNode->SetDisplayLabel(NewLabel);
+		}
+
+		//FactoryNode->SetEnabled(true);
+	};
+
+	//Find all factory node we need for this pipeline
+	NodeContainer->IterateNodesOfType<UInterchangeFactoryBaseNode>(UpdateFactoryNodes);
+}
+
+void UInterchangeDataprepLevelPipeline::ExecutePostImportPipeline(const UInterchangeBaseNodeContainer* NodeContainer, const FString& FactoryNodeKey, UObject* CreatedAsset, bool bIsAReimport)
+{
+	if (!NodeContainer || !CreatedAsset)
+	{
+		return;
+	}
+
+	Super::ExecutePostImportPipeline(NodeContainer, FactoryNodeKey, CreatedAsset, bIsAReimport);
+
+	CreatedAsset->ClearFlags(RF_Public/* | RF_Standalone*/);
+	CreatedAsset->SetFlags(RF_Transient);
+}
diff --git a/Engine/Plugins/Enterprise/DataprepEditor/Source/DataprepCore/Private/InterchangeDataprepPipeline.h b/Engine/Plugins/Enterprise/DataprepEditor/Source/DataprepCore/Private/InterchangeDataprepPipeline.h
new file mode 100644
--- /dev/null
+++ b/Engine/Plugins/Enterprise/DataprepEditor/Source/DataprepCore/Private/InterchangeDataprepPipeline.h
@@ -0,0 +1,20 @@
+// Copyright Epic Games, Inc. All Rights Reserved.
+
+#pragma once
+
+#include "CoreMinimal.h"
+#include "InterchangePipelineBase.h"
+#include "UObject/Object.h"
+#include "UObject/ObjectMacros.h"
+
+#include "InterchangeDataprepPipeline.generated.h"
+
+UCLASS(BlueprintType)
+class UInterchangeDataprepLevelPipeline : public UInterchangePipelineBase
+{
+	GENERATED_BODY()
+
+protected:
+	virtual void ExecutePipeline(UInterchangeBaseNodeContainer* BaseNodeContainer, const TArray<UInterchangeSourceData*>& SourceDatas) override;
+	virtual void ExecutePostImportPipeline(const UInterchangeBaseNodeContainer* BaseNodeContainer, const FString& NodeKey, UObject* CreatedAsset, bool bIsAReimport) override;
+};
\ No newline at end of file
diff --git a/Engine/Plugins/Enterprise/DataprepEditor/Source/DataprepCore/Private/InterchangeFileProducer.cpp b/Engine/Plugins/Enterprise/DataprepEditor/Source/DataprepCore/Private/InterchangeFileProducer.cpp
new file mode 100644
--- /dev/null
+++ b/Engine/Plugins/Enterprise/DataprepEditor/Source/DataprepCore/Private/InterchangeFileProducer.cpp
@@ -0,0 +1,664 @@
+// Copyright Epic Games, Inc. All Rights Reserved.
+
+#include "InterchangeFileProducer.h"
+
+#include "DataprepAssetProducers.h"
+#include "IDataprepLogger.h"
+#include "IDataprepProgressReporter.h"
+#include "InterchangeDataprepPipeline.h"
+
+#include "InterchangeManager.h"
+
+#include "AssetRegistry/AssetRegistryModule.h"
+#include "AssetRegistry/ARFilter.h"
+#include "DesktopPlatformModule.h"
+#include "DetailCategoryBuilder.h"
+#include "DetailLayoutBuilder.h"
+#include "DetailWidgetRow.h"
+#include "EditorDirectories.h"
+#include "EditorFontGlyphs.h"
+#include "Framework/Application/SlateApplication.h"
+#include "Misc/PackageName.h"
+#include "Modules/ModuleManager.h"
+#include "ObjectTools.h"
+#include "ScopedTransaction.h"
+#include "Widgets/Images/SImage.h"
+#include "Widgets/Input/SCheckBox.h"
+#include "Widgets/Input/SButton.h"
+#include "Widgets/Input/SEditableTextBox.h"
+#include "Widgets/Layout/SBox.h"
+#include "Widgets/SBoxPanel.h"
+#include "Widgets/SWidget.h"
+#include "Widgets/Text/STextBlock.h"
+#include "Widgets/Layout/SUniformGridPanel.h"
+#include "Widgets/Layout/SGridPanel.h"
+#include "Widgets/Text/SInlineEditableTextBlock.h"
+
+#define LOCTEXT_NAMESPACE "InterchangeFileProducer"
+
+const FText InterchangeFileProducerLabel( LOCTEXT( "InterchangeFileProducerLabel", "Interchange file importer" ) );
+const FText InterchangeFileProducerDescription( LOCTEXT( "InterchangeFileProducerDesc", "Reads a file supported by Interchange and its dependent assets" ) );
+
+const TCHAR* WildCharCriteria = TEXT( "*.*" );
+const TCHAR* ExtensionPrefix = TEXT( "*." );
+
+bool bGEnableInterchangeProducer = false;
+FAutoConsoleVariableRef GEnableInterchangeProducerVar(
+	TEXT("Dataprep.Interchange.Producer"),
+	bGEnableInterchangeProducer,
+	TEXT("Enable/disable the Intrchange file producer.\nDefault is disabled\n"),
+	ECVF_Default
+);
+
+TArray<FSoftObjectPath> UInterchangeFileProducer::OverridePipelines;
+
+namespace FInterchangeFileProducerUtils
+{
+	/** Delete all the assets stored under the specified path */
+	void DeletePackagePath( const FString& PathToDelete );
+
+	/** Delete all the packages created by the Interchange importer */
+	void DeletePackagesPath(TSet<FString>& PathsToDelete)
+	{
+		for(FString& PathToDelete : PathsToDelete)
+		{
+			DeletePackagePath( PathToDelete );
+		}
+	}
+
+	/** Display OS browser, i.e. Windows explorer, to let user select a file */
+	FString SelectFileToImport();
+
+	/** Display OS browser, i.e. Windows explorer, to let user select a directory */
+	FString SelectDirectory();
+
+
+	class SFileProducerImportOptionsWindow : public SCompoundWidget
+	{
+	public:
+		SLATE_BEGIN_ARGS(SFileProducerImportOptionsWindow)
+		{}
+		SLATE_ARGUMENT(TArray<UObject*>, ImportOptions)
+		SLATE_ARGUMENT(TSharedPtr<SWindow>, WidgetWindow)
+		SLATE_ARGUMENT(int32, MinDetailHeight)
+		SLATE_ARGUMENT(int32, MinDetailWidth)
+		SLATE_END_ARGS()
+
+	public:
+		void Construct(const FArguments& InArgs)
+		{
+			ImportOptions = InArgs._ImportOptions;
+			Window = InArgs._WidgetWindow;
+
+			TSharedPtr<SBox> DetailsViewBox;
+
+			ChildSlot
+			[
+				SNew(SVerticalBox)
+				+ SVerticalBox::Slot()
+				.AutoHeight()
+				.Padding(2)
+				[
+					SAssignNew(DetailsViewBox, SBox)
+					.MinDesiredHeight(InArgs._MinDetailHeight)
+					.MinDesiredWidth(InArgs._MinDetailWidth)
+				]
+			];
+
+			FPropertyEditorModule& PropertyEditorModule = FModuleManager::GetModuleChecked<FPropertyEditorModule>("PropertyEditor");
+			FDetailsViewArgs DetailsViewArgs;
+			DetailsViewArgs.bAllowSearch = false;
+			DetailsViewArgs.bAllowMultipleTopLevelObjects = true;
+			DetailsViewArgs.NameAreaSettings = FDetailsViewArgs::HideNameArea;
+			TSharedPtr<IDetailsView> DetailsView = PropertyEditorModule.CreateDetailView(DetailsViewArgs);
+
+			DetailsViewBox->SetContent(DetailsView.ToSharedRef());
+			DetailsView->SetObjects(ImportOptions);
+		}
+
+		virtual bool SupportsKeyboardFocus() const override { return true; }
+
+		virtual FReply OnKeyDown(const FGeometry& MyGeometry, const FKeyEvent& InKeyEvent) override
+		{
+			if (InKeyEvent.GetKey() == EKeys::Escape)
+			{
+				if (Window.IsValid())
+				{
+					Window.Pin()->RequestDestroyWindow();
+				}
+				return FReply::Handled();
+			}
+
+			return FReply::Unhandled();
+		}
+
+	private:
+		TArray<UObject*> ImportOptions;
+		TWeakPtr< SWindow > Window;
+	};
+}
+
+UInterchangeFileProducer::UInterchangeFileProducer()
+{
+	if (HasAnyFlags(RF_ClassDefaultObject))
+	{
+		// #dataprep_todo Allow users to override the stack of pipelines
+		OverridePipelines.Add(FSoftObjectPath(TEXT("/Interchange/Pipelines/DefaultSceneAssetsPipeline.DefaultSceneAssetsPipeline")));
+		OverridePipelines.Add(FSoftObjectPath(TEXT("/Interchange/Pipelines/DefaultSceneLevelPipeline.DefaultSceneLevelPipeline")));
+		OverridePipelines.Add(FSoftObjectPath(TEXT("/DataprepEditor/DefaultDataprepPipeline.DefaultDataprepPipeline")));
+	}
+}
+
+bool UInterchangeFileProducer::IsActive()
+{
+	return bGEnableInterchangeProducer;
+}
+
+bool UInterchangeFileProducer::Initialize()
+{
+	FText TaskDescription = FText::Format( LOCTEXT( "InterchangeFileProducer_LoadingFile", "Loading {0} ..."), FText::FromString( FilePath ) );
+	ProgressTaskPtr = MakeUnique< FDataprepWorkReporter >( Context.ProgressReporterPtr, TaskDescription, 10.0f, 1.0f );
+
+	ProgressTaskPtr->ReportNextStep( TaskDescription, 7.0f );
+
+	if( FilePath.IsEmpty() )
+	{
+		LogError( LOCTEXT( "InterchangeFileProducer_Incomplete", "No file has been selected." ) );
+		return false;
+	}
+
+	// Check file exists
+	if(!FPaths::FileExists( FilePath ))
+	{
+		LogError( FText::Format( LOCTEXT( "InterchangeFileProducer_NotFound", "File {0} does not exist." ), FText::FromString( FilePath ) ) );
+		return false;
+	}
+
+	// Create transient package if it wasn't specified
+	if (!TransientPackage)
+	{
+		TransientPackage = NewObject< UPackage >(nullptr, *FPaths::Combine(Context.RootPackagePtr->GetPathName(), *GetName()), RF_Transient);
+	}
+
+	TransientPackage->FullyLoad();
+
+	constexpr EObjectFlags LocalObjectFlags = RF_Public | RF_Standalone | RF_Transactional;
+	FFeedbackContext* InFeedbackContext = Context.ProgressReporterPtr ? Context.ProgressReporterPtr->GetFeedbackContext() : nullptr;
+
+	return true;
+}
+
+bool UInterchangeFileProducer::InitTranslator()
+{
+	return true;
+}
+
+bool UInterchangeFileProducer::Execute(TArray< TWeakObjectPtr< UObject > >& OutAssets)
+{
+	using namespace UE::Interchange;
+	using FInterchangeResults = TTuple<FAssetImportResultRef, FSceneImportResultRef>;
+
+	if ( !IsValid() || IsCancelled() )
+	{
+		return false;
+	}
+
+	ProgressTaskPtr->ReportNextStep( FText::Format( LOCTEXT( "InterchangeFileProducer_ConvertingFile", "Converting {0} ..."), FText::FromString( FilePath ) ), 2.0f );
+
+	UInterchangeManager& InterchangeManager = UInterchangeManager::GetInterchangeManager();
+
+	FScopedSourceData ScopedSourceData(FilePath);
+
+	if (!InterchangeManager.CanTranslateSourceData(ScopedSourceData.GetSourceData()))
+	{
+		return false;
+	}
+	
+	FImportAssetParameters ImportAssetParameters;
+	
+	ImportAssetParameters.bIsAutomated = true;
+	ImportAssetParameters.ImportLevel = Context.WorldPtr.Get()->GetCurrentLevel();
+	ImportAssetParameters.OverridePipelines = OverridePipelines;
+
+	FInterchangeResults ImportResults = InterchangeManager.ImportSceneAsync(TransientPackage->GetPathName(), ScopedSourceData.GetSourceData(), ImportAssetParameters);
+
+	ImportResults.Get<0>()->WaitUntilDone();
+	ImportResults.Get<1>()->WaitUntilDone();
+
+	if (!ImportResults.Key->IsValid() || !ImportResults.Value->IsValid())
+	{
+		return false;
+	}
+
+	OutAssets.Reserve(OutAssets.Num() + ImportResults.Get<1>()->GetImportedObjects().Num());
+
+	OutAssets.Append(ImportResults.Get<0>()->GetImportedObjects());
+
+	for (UObject* Object : ImportResults.Get<1>()->GetImportedObjects())
+	{
+		if (Object && !Object->GetClass()->IsChildOf(AActor::StaticClass()))
+		{
+			OutAssets.Add(Object);
+		}
+	}
+
+	return !IsCancelled();
+}
+
+void UInterchangeFileProducer::OnFilePathChanged()
+{
+	FilePath = FPaths::ConvertRelativePathToFull( FilePath );
+	UpdateName();
+	OnChanged.Broadcast( this );
+}
+
+void UInterchangeFileProducer::Reset()
+{
+	ProgressTaskPtr.Reset();
+	Assets.Empty();
+	TransientPackage = nullptr;
+
+	UDataprepContentProducer::Reset();
+}
+
+const FText& UInterchangeFileProducer::GetLabel() const
+{
+	return InterchangeFileProducerLabel;
+}
+
+const FText& UInterchangeFileProducer::GetDescription() const
+{
+	return InterchangeFileProducerDescription;
+}
+
+FString UInterchangeFileProducer::GetNamespace() const
+{
+	return {};
+}
+
+void UInterchangeFileProducer::SetFilePath( const FString& InFilePath )
+{
+	Modify();
+
+	FilePath = InFilePath;
+
+	OnFilePathChanged();
+}
+
+void UInterchangeFileProducer::UpdateName()
+{
+	if(!FilePath.IsEmpty())
+	{
+		// Rename producer to name of file
+		FString CleanName = ObjectTools::SanitizeObjectName( FPaths::GetCleanFilename( FilePath ) );
+		if ( !Rename( *CleanName, nullptr, REN_Test ) )
+		{
+			CleanName = MakeUniqueObjectName( GetOuter(), GetClass(), *CleanName ).ToString();
+		}
+
+		Rename( *CleanName, nullptr, REN_DontCreateRedirectors | REN_NonTransactional );
+	}
+}
+
+bool UInterchangeFileProducer::Supersede(const UDataprepContentProducer* OtherProducer) const
+{
+	const UInterchangeFileProducer* OtherFileProducer = Cast<const UInterchangeFileProducer>(OtherProducer);
+
+	return OtherFileProducer != nullptr &&
+		!OtherFileProducer->FilePath.IsEmpty() &&
+		FilePath == OtherFileProducer->FilePath;
+}
+
+bool UInterchangeFileProducer::CanAddToProducersArray(bool bIsAutomated)
+{
+	if ( bIsAutomated )
+	{
+		return true;
+	}
+
+	FilePath = FInterchangeFileProducerUtils::SelectFileToImport();
+	if ( !FilePath.IsEmpty() )
+	{
+		UpdateName();
+		return true;
+	}
+
+	return false;
+}
+
+void UInterchangeFileProducer::PostEditUndo()
+{
+	UDataprepContentProducer::PostEditUndo();
+
+	OnChanged.Broadcast( this );
+}
+
+void UInterchangeFileProducer::PostEditChangeProperty(FPropertyChangedEvent& PropertyChangedEvent)
+{
+	FProperty* EditedProperty = PropertyChangedEvent.Property;
+
+	if ( !EditedProperty || EditedProperty->GetFName() == GET_MEMBER_NAME_CHECKED(UInterchangeFileProducer, FilePath) )
+	{
+		OnFilePathChanged();
+	}
+}
+
+class SInterchangeFileProducerFileProperty : public SCompoundWidget
+{
+public:
+	SLATE_BEGIN_ARGS(SInterchangeFileProducerFileProperty)
+	{}
+
+	SLATE_ARGUMENT(UInterchangeFileProducer*, Producer)
+	SLATE_END_ARGS()
+
+public:
+	void Construct(const FArguments& InArgs)
+	{
+		ProducerPtr = TWeakObjectPtr< UInterchangeFileProducer >( InArgs._Producer );
+
+		FSlateFontInfo FontInfo = IDetailLayoutBuilder::GetDetailFont();
+
+		ChildSlot
+		[
+			SNew(SHorizontalBox)
+			+ SHorizontalBox::Slot()
+			.VAlign(VAlign_Center)
+			[
+				SAssignNew(FileName, SEditableText)
+				.IsReadOnly(true)
+				.Text(this, &SInterchangeFileProducerFileProperty::GetFilenameText)
+				.ToolTipText(this, &SInterchangeFileProducerFileProperty::GetFilenameText)
+				.Font( FontInfo )
+			]
+			+SHorizontalBox::Slot()
+			.VAlign(VAlign_Center)
+			.AutoWidth()
+			[
+				SNew(SButton)
+				.OnClicked(this, &SInterchangeFileProducerFileProperty::OnChangePathClicked )
+				.ToolTipText(LOCTEXT("ChangeSourcePath_Tooltip", "Browse for a new source file path"))
+				[
+					SNew(STextBlock)
+					.Text(LOCTEXT("...", "..."))
+					.Font( FontInfo )
+				]
+			]
+		];
+	}
+
+private:
+	FReply OnChangePathClicked() const
+	{
+		if( !ProducerPtr.IsValid() )
+		{
+			return FReply::Unhandled();
+		}
+
+		FString SelectedFile = FInterchangeFileProducerUtils::SelectFileToImport();
+		if(!SelectedFile.IsEmpty())
+		{
+			const FScopedTransaction Transaction( LOCTEXT("Producer_SetFilename", "Set Filename") );
+
+			ProducerPtr->SetFilePath( SelectedFile );
+			FileName->SetText( GetFilenameText() );
+		}
+
+		return FReply::Handled();
+	}
+
+	FText GetFilenameText() const
+	{
+		return ProducerPtr->FilePath.IsEmpty() ? FText::FromString( TEXT("Select a file") ) : FText::FromString( ProducerPtr->FilePath );
+	}
+
+private:
+	TWeakObjectPtr< UInterchangeFileProducer > ProducerPtr;
+	TSharedPtr< SEditableText > FileName;
+};
+
+void FDataprepContentProducerDetails::CustomizeDetails(IDetailLayoutBuilder& DetailBuilder)
+{
+	TArray< TWeakObjectPtr< UObject > > Objects;
+	DetailBuilder.GetObjectsBeingCustomized(Objects);
+	check(Objects.Num() > 0);
+
+	Producer = Cast< UDataprepContentProducer >(Objects[0].Get());
+	check(Producer);
+
+	AssetProducers = Cast<UDataprepAssetProducers>(Producer->GetOuter());
+	check(AssetProducers);
+
+	ProducerIndex = INDEX_NONE;
+	for (int Index = 0; Index < AssetProducers->GetProducersCount(); ++Index)
+	{
+		if (Producer == AssetProducers->GetProducer(Index))
+		{
+			ProducerIndex = Index;
+			break;
+		}
+	}
+}
+
+FSlateColor FDataprepContentProducerDetails::GetStatusColorAndOpacity() const
+{
+	return  IsProducerSuperseded() ? FLinearColor::Red : FAppStyle::Get().GetSlateColor("DefaultForeground");
+}
+
+bool FDataprepContentProducerDetails::IsProducerSuperseded() const
+{
+	return ProducerIndex != INDEX_NONE ? AssetProducers->IsProducerSuperseded(ProducerIndex) : false;
+}
+
+void FInterchangeFileProducerDetails::CustomizeDetails(IDetailLayoutBuilder& DetailBuilder)
+{
+	FDataprepContentProducerDetails::CustomizeDetails(DetailBuilder);
+
+	UInterchangeFileProducer* FileProducer = Cast< UInterchangeFileProducer >(Producer);
+	check( FileProducer );
+
+	// #ueent_todo: Remove handling of warning category when this is not considered experimental anymore
+	TArray<FName> CategoryNames;
+	DetailBuilder.GetCategoryNames( CategoryNames );
+	CategoryNames.Remove( FName(TEXT("Warning")) );
+
+	DetailBuilder.HideCategory(FName( TEXT( "Warning" ) ) );
+
+	//FName CategoryName( TEXT("InterchangeFileProducerCustom") );
+	IDetailCategoryBuilder& ImportSettingsCategoryBuilder = DetailBuilder.EditCategory( NAME_None, FText::GetEmpty(), ECategoryPriority::Important );
+
+	TSharedRef<IPropertyHandle> PropertyHandle = DetailBuilder.GetProperty( TEXT("FilePath"), UInterchangeFileProducer::StaticClass() );
+	PropertyHandle->MarkHiddenByCustomization();
+
+	FDetailWidgetRow& CustomAssetImportRow = ImportSettingsCategoryBuilder.AddCustomRow( FText::FromString( TEXT( "Import File" ) ) );
+
+	TSharedPtr<STextBlock> IconText;
+
+	CustomAssetImportRow.NameContent()
+	[
+		SNew(SHorizontalBox)
+		+ SHorizontalBox::Slot()
+		.VAlign(VAlign_Center)
+		.AutoWidth()
+		.Padding(0, 0, 3, 0)
+		[
+			SNew(SButton)
+			.ButtonStyle(FAppStyle::Get(), "HoverHintOnly")
+			.IsFocusable(false)
+			.IsEnabled(false) // Options button
+			.OnClicked_Lambda( [FileProducer]() -> FReply 
+			{ 
+				//FileProducer->OnChangeImportSettings(); 
+				return FReply::Handled(); 
+			})
+			.ToolTipText(LOCTEXT("ChangeImportSettings_Tooltip", "Import Settings"))
+			[
+				SNew(STextBlock)
+				.Font(FAppStyle::Get().GetFontStyle("FontAwesome.11"))
+				.ColorAndOpacity(FLinearColor::White)
+				.Text(FEditorFontGlyphs::Cog)
+			]
+		]
+
+		+ SHorizontalBox::Slot()
+		.VAlign(VAlign_Center)
+		.AutoWidth()
+		[
+			SAssignNew(IconText, STextBlock)
+				.Font(FAppStyle::Get().GetFontStyle("FontAwesome.11"))
+				.Text(MakeAttributeLambda([this]
+				{
+					return this->IsProducerSuperseded() ? FEditorFontGlyphs::Exclamation_Triangle : FEditorFontGlyphs::File;
+				}))
+				.ColorAndOpacity(this, &FDataprepContentProducerDetails::GetStatusColorAndOpacity)
+		]
+	];
+
+	IconText->SetToolTipText(MakeAttributeLambda([this]
+	{
+		if (IsProducerSuperseded())
+		{
+			return LOCTEXT("FInterchangeFileProducerDetails_StatusTextTooltip_Superseded", "This producer is superseded by another one and will be skipped when run.");
+		}
+		return LOCTEXT("FInterchangeFileProducerDetails_StatusTextTooltip", "File input");
+	}));
+
+	CustomAssetImportRow.ValueContent()
+	.MinDesiredWidth( 2000.0f )
+	[
+		SNew( SInterchangeFileProducerFileProperty )
+		.Producer( FileProducer )
+	];
+}
+
+void FInterchangeFileProducerUtils::DeletePackagePath( const FString& PathToDelete )
+{
+	if(PathToDelete.IsEmpty())
+	{
+		return;
+	}
+
+	// Inspired from ContentBrowserUtils::DeleteFolders
+	// Inspired from ContentBrowserUtils::LoadAssetsIfNeeded
+
+	// Form a filter from the paths
+	FAssetRegistryModule& AssetRegistryModule = FModuleManager::Get().LoadModuleChecked<FAssetRegistryModule>(TEXT("AssetRegistry"));
+
+	FARFilter Filter;
+	Filter.bRecursivePaths = true;
+	new (Filter.PackagePaths) FName(*PathToDelete);
+
+	// Query for a list of assets in the selected paths
+	TArray<FAssetData> AssetDataList;
+	AssetRegistryModule.Get().GetAssets(Filter, AssetDataList);
+
+	bool bAllowFolderDelete = false;
+
+	{
+		struct FEmptyFolderVisitor : public IPlatformFile::FDirectoryVisitor
+		{
+			bool bIsEmpty;
+
+			FEmptyFolderVisitor()
+				: bIsEmpty(true)
+			{
+			}
+
+			virtual bool Visit(const TCHAR* FilenameOrDirectory, bool bIsDirectory) override
+			{
+				if (!bIsDirectory)
+				{
+					bIsEmpty = false;
+					return false; // abort searching
+				}
+
+				return true; // continue searching
+			}
+		};
+
+		FString PathToDeleteOnDisk;
+		if (FPackageName::TryConvertLongPackageNameToFilename(PathToDelete, PathToDeleteOnDisk))
+		{
+			// Look for files on disk in case the folder contains things not tracked by the asset registry
+			FEmptyFolderVisitor EmptyFolderVisitor;
+			IFileManager::Get().IterateDirectoryRecursively(*PathToDeleteOnDisk, EmptyFolderVisitor);
+
+			if (EmptyFolderVisitor.bIsEmpty && IFileManager::Get().DeleteDirectory(*PathToDeleteOnDisk, false, true))
+			{
+				AssetRegistryModule.Get().RemovePath(PathToDelete);
+			}
+		}
+	}
+}
+
+FString FInterchangeFileProducerUtils::SelectFileToImport()
+{
+	TArray<FString> Formats = UInterchangeManager::GetInterchangeManager().GetSupportedFormats(EInterchangeTranslatorType::Scenes);
+
+	FString FileTypes;
+	FString AllExtensions;
+
+	for( const FString& Format : Formats )
+	{
+		TArray<FString> FormatComponents;
+		Format.ParseIntoArray( FormatComponents, TEXT( ";" ), false );
+
+		for ( int32 ComponentIndex = 0; ComponentIndex < FormatComponents.Num(); ComponentIndex += 2 )
+		{
+			check( FormatComponents.IsValidIndex( ComponentIndex + 1 ) );
+			const FString& Extension = FormatComponents[ComponentIndex];
+			const FString& Description = FormatComponents[ComponentIndex + 1];
+
+			if ( !AllExtensions.IsEmpty() )
+			{
+				AllExtensions.AppendChar( TEXT( ';' ) );
+			}
+			AllExtensions.Append( TEXT( "*." ) );
+			AllExtensions.Append( Extension );
+
+			if ( !FileTypes.IsEmpty() )
+			{
+				FileTypes.AppendChar( TEXT( '|' ) );
+			}
+
+			FileTypes.Append( FString::Printf( TEXT( "%s (*.%s)|*.%s" ), *Description, *Extension, *Extension ) );
+		}
+	}
+
+	FString SupportedExtensions( FString::Printf( TEXT( "All Files (%s)|%s|%s" ), *AllExtensions, *AllExtensions, *FileTypes ) );
+
+	TArray<FString> OpenedFiles;
+	FString DefaultLocation( FEditorDirectories::Get().GetLastDirectory( ELastDirectory::GENERIC_IMPORT ) );
+	IDesktopPlatform* DesktopPlatform = FDesktopPlatformModule::Get();
+
+	bool bOpened = false;
+	if ( DesktopPlatform )
+	{
+		bOpened = DesktopPlatform->OpenFileDialog(
+			FSlateApplication::Get().FindBestParentWindowHandleForDialogs( nullptr ),
+			LOCTEXT( "FileDialogTitle", "Import Interchange" ).ToString(),
+			DefaultLocation,
+			TEXT( "" ),
+			SupportedExtensions,
+			EFileDialogFlags::None,
+			OpenedFiles
+		);
+	}
+
+	if ( bOpened && OpenedFiles.Num() > 0 )
+	{
+		const FString& OpenedFile = OpenedFiles[0];
+		FEditorDirectories::Get().SetLastDirectory( ELastDirectory::GENERIC_IMPORT, FPaths::GetPath(OpenedFile) );
+
+		return FPaths::ConvertRelativePathToFull( OpenedFile );
+	}
+
+	return {};
+}
+
+FString FInterchangeFileProducerUtils::SelectDirectory()
+{
+	return {};
+}
+
+#undef LOCTEXT_NAMESPACE
diff --git a/Engine/Plugins/Enterprise/DataprepEditor/Source/DataprepCore/Private/InterchangeFileProducer.h b/Engine/Plugins/Enterprise/DataprepEditor/Source/DataprepCore/Private/InterchangeFileProducer.h
new file mode 100644
--- /dev/null
+++ b/Engine/Plugins/Enterprise/DataprepEditor/Source/DataprepCore/Private/InterchangeFileProducer.h
@@ -0,0 +1,96 @@
+// Copyright Epic Games, Inc. All Rights Reserved.
+
+#pragma once
+
+#include "CoreMinimal.h"
+
+#include "DataprepContentProducer.h"
+
+#include "IDetailCustomization.h"
+#include "Input/Reply.h"
+#include "Styling/SlateColor.h"
+
+#include "InterchangeFileProducer.generated.h"
+
+class IDetailLayoutBuilder;
+class UInterchangePipelineBase;
+
+UCLASS(Experimental, HideCategories = (InterchangeProducer_Internal))
+class DATAPREPCORE_API UInterchangeFileProducer : public UDataprepContentProducer
+{
+	GENERATED_BODY()
+
+	UInterchangeFileProducer();
+
+public:
+	// UObject interface
+	virtual void PostEditUndo() override;
+	virtual void PostEditChangeProperty(FPropertyChangedEvent& PropertyChangedEvent) override;
+	// End of UObject interface
+
+	/** Update producer with the desired file */
+	void SetFilePath( const FString& InFilePath );
+
+	const FString& GetFilePath() const { return FilePath; }
+
+	// Begin UDataprepContentProducer overrides
+	virtual const FText& GetLabel() const override;
+	virtual const FText& GetDescription() const override;
+	virtual FString GetNamespace() const override;
+	virtual bool Supersede(const UDataprepContentProducer* OtherProducer) const override;
+	virtual bool CanAddToProducersArray(bool bIsAutomated) override;
+
+protected:
+	virtual bool Initialize() override;
+	virtual bool Execute(TArray< TWeakObjectPtr< UObject > >& OutAssets) override;
+	virtual void Reset() override;
+	virtual bool IsActive() override;
+	// End UDataprepContentProducer overrides
+
+	UPROPERTY( EditAnywhere, Category = InterchangeProducer )
+	FString FilePath;
+
+private:
+	/** Does what is required after setting a new FilePath */
+	void OnFilePathChanged();
+
+	/** Update the name of the producer based on the filename */
+	void UpdateName();
+
+	bool InitTranslator();
+
+private:
+	TUniquePtr< FDataprepWorkReporter > ProgressTaskPtr;
+
+	UPROPERTY( Transient, DuplicateTransient )
+	TObjectPtr<UPackage> TransientPackage = nullptr;
+
+	static TArray<FSoftObjectPath> OverridePipelines;
+
+	TArray< TWeakObjectPtr< UObject > > Assets;
+
+	friend class SInterchangeFileProducerFileProperty;
+};
+
+class FDataprepContentProducerDetails : public IDetailCustomization
+{
+public:
+	virtual void CustomizeDetails( IDetailLayoutBuilder& DetailBuilder ) override;
+	FSlateColor GetStatusColorAndOpacity() const;
+	bool IsProducerSuperseded() const;
+
+protected:
+	class UDataprepAssetProducers* AssetProducers = nullptr;
+	UDataprepContentProducer* Producer = nullptr;
+	int32 ProducerIndex = INDEX_NONE; // This index does not change for the lifetime of the property widget
+};
+
+// Customization of the details of the Interchange producer for the data prep editor.
+class DATAPREPCORE_API FInterchangeFileProducerDetails : public FDataprepContentProducerDetails
+{
+public:
+	static TSharedRef< IDetailCustomization > MakeDetails() { return MakeShared<FInterchangeFileProducerDetails>(); };
+
+	/** Called when details should be customized */
+	virtual void CustomizeDetails( IDetailLayoutBuilder& DetailBuilder ) override;
+};
diff --git a/Engine/Plugins/Enterprise/DataprepEditor/Source/DataprepCore/Public/DataprepContentProducer.h b/Engine/Plugins/Enterprise/DataprepEditor/Source/DataprepCore/Public/DataprepContentProducer.h
--- a/Engine/Plugins/Enterprise/DataprepEditor/Source/DataprepCore/Public/DataprepContentProducer.h
+++ b/Engine/Plugins/Enterprise/DataprepEditor/Source/DataprepCore/Public/DataprepContentProducer.h
@@ -124,6 +124,15 @@ public:
 		return true;
 	}
 
+	/**
+	 * Allow a producer to pop a ui after being created.
+	 * @return false to cancel the insertion of the new producer
+	 */
+	virtual bool IsActive()
+	{
+		return true;
+	}
+
 	/** Returns true if the producer was cancelled during execution */
 	bool IsCancelled() { return Context.ProgressReporterPtr.IsValid() ? Context.ProgressReporterPtr->IsWorkCancelled() : false; }
 
diff --git a/Engine/Plugins/Enterprise/DataprepEditor/Source/DataprepEditor/Private/DataprepSnapshot.cpp b/Engine/Plugins/Enterprise/DataprepEditor/Source/DataprepEditor/Private/DataprepSnapshot.cpp
--- a/Engine/Plugins/Enterprise/DataprepEditor/Source/DataprepEditor/Private/DataprepSnapshot.cpp
+++ b/Engine/Plugins/Enterprise/DataprepEditor/Source/DataprepEditor/Private/DataprepSnapshot.cpp
@@ -295,6 +295,27 @@ namespace DataprepSnapshotUtil
 				SubObjectDependencyGraph[SubObject].Append(Analyzer.DependentObjects);
 			}
 
+			bool bCyclingDependency = true;
+			for (auto& Entry : SubObjectDependencyGraph)
+			{
+				if (Entry.Value.Num() == 0)
+				{
+					bCyclingDependency = false;
+					break;
+				}
+			}
+
+			if (bCyclingDependency)
+			{
+				for (auto& Entry : SubObjectDependencyGraph)
+				{
+					if (Entry.Value.Num() == 1)
+					{
+						Entry.Value.Empty();
+					}
+				}
+			}
+
 			// Sort array of sub-objects: first objects do not depend on ones below
 			int32 Count = SubObjectsArray.Num();
 			SubObjectsArray.Empty(Count);
diff --git a/Engine/Plugins/Enterprise/DataprepEditor/Source/DataprepEditor/Private/Widgets/SDataprepProducersWidget.cpp b/Engine/Plugins/Enterprise/DataprepEditor/Source/DataprepEditor/Private/Widgets/SDataprepProducersWidget.cpp
--- a/Engine/Plugins/Enterprise/DataprepEditor/Source/DataprepEditor/Private/Widgets/SDataprepProducersWidget.cpp
+++ b/Engine/Plugins/Enterprise/DataprepEditor/Source/DataprepEditor/Private/Widgets/SDataprepProducersWidget.cpp
@@ -308,7 +308,6 @@ TSharedRef<SWidget> SDataprepProducersWidget::CreateAddProducerMenuWidget()
 	MenuBuilder.BeginSection("AddNewProducer", LOCTEXT("DataprepProducersWidget_AddImports", "Add Producer"));
 	{
 		FUIAction MenuAction;
-		int32 Index = 0;
 
 		// Find content producers the user could use for their data preparation
 		for( TObjectIterator< UClass > It ; It ; ++It )
@@ -319,10 +318,13 @@ TSharedRef<SWidget> SDataprepProducersWidget::CreateAddProducerMenuWidget()
 			{
 				if( CurrentClass->IsChildOf( UDataprepContentProducer::StaticClass() ) )
 				{
-					MenuAction.ExecuteAction.BindSP(this, &SDataprepProducersWidget::OnAddProducer, CurrentClass);
-
 					UDataprepContentProducer* DefaultProducer = CurrentClass->GetDefaultObject<UDataprepContentProducer>();
-					check( DefaultProducer );
+					if (!ensure(DefaultProducer) || !DefaultProducer->IsActive())
+					{
+						continue;
+					}
+
+					MenuAction.ExecuteAction.BindSP(this, &SDataprepProducersWidget::OnAddProducer, CurrentClass);
 
 					MenuBuilder.AddMenuEntry(
 						DefaultProducer->GetLabel(),
@@ -332,8 +334,6 @@ TSharedRef<SWidget> SDataprepProducersWidget::CreateAddProducerMenuWidget()
 						NAME_None,
 						EUserInterfaceActionType::Button
 					);
-
-					++Index;
 				}
 			}
 		}
diff --git a/Engine/Plugins/Enterprise/DatasmithImporter/Source/DatasmithImporter/Private/DatasmithConsumer.cpp b/Engine/Plugins/Enterprise/DatasmithImporter/Source/DatasmithImporter/Private/DatasmithConsumer.cpp
--- a/Engine/Plugins/Enterprise/DatasmithImporter/Source/DatasmithImporter/Private/DatasmithConsumer.cpp
+++ b/Engine/Plugins/Enterprise/DatasmithImporter/Source/DatasmithImporter/Private/DatasmithConsumer.cpp
@@ -383,8 +383,6 @@ bool UDatasmithConsumer::Initialize()
 			FPropertyChangedEvent PropertyUpdateStruct( Property );
 			PostEditChangeProperty( PropertyUpdateStruct );
 
-			MarkPackageDirty();
-
 			DatasmithSceneObjectPath = FSoftObjectPath(DatasmithSceneWeakPtr.Get()).GetAssetPathString();
 		}
 	}
@@ -1161,6 +1159,16 @@ bool UDatasmithConsumer::ValidateAssets()
 	{
 		if(UObject* Asset = AssetPtr.Get())
 		{
+			// If the asset has not been created by a Datasmith producer
+			const FString SourcePackagePath = Asset->GetOutermost()->GetName();
+			if (SourcePackagePath.StartsWith(Context.TransientContentFolder, ESearchCase::CaseSensitive))
+			{
+				const FString DestinationPackagePath = SourcePackagePath.Replace(*Context.TransientContentFolder, *TargetContentFolder, ESearchCase::CaseSensitive);
+
+				Asset->ClearFlags(RF_Transient);
+				Asset = FDatasmithImporterImpl::FinalizeAsset(Asset, *DestinationPackagePath, nullptr, nullptr);
+			}
+
 			TPair<FString, FString>& AssetPackageInfo = AssetPackageInfoMap.Add(Asset);
 
 			bool bAssetWithMarker = false;
diff --git a/Engine/Plugins/Enterprise/DatasmithImporter/Source/DatasmithImporter/Private/DatasmithImporter.cpp b/Engine/Plugins/Enterprise/DatasmithImporter/Source/DatasmithImporter/Private/DatasmithImporter.cpp
--- a/Engine/Plugins/Enterprise/DatasmithImporter/Source/DatasmithImporter/Private/DatasmithImporter.cpp
+++ b/Engine/Plugins/Enterprise/DatasmithImporter/Source/DatasmithImporter/Private/DatasmithImporter.cpp
@@ -1537,7 +1537,7 @@ AActor* FDatasmithImporter::FinalizeActor( FDatasmithImportContext& ImportContex
 
 		// Actors that were created via blueprint (can happen with Dataprep) will have their components destroyed by the finalizer after this scope's end.
 		// This will disallow destruction of such components.
-		DestinationActor->bActorSeamlessTraveled = 1;
+		//DestinationActor->bActorSeamlessTraveled = 1;
 	}
 
 	// Update label to match the source actor's
