Fixed linkage error when compiling monolithic

#rnx

diff --git a/Engine/Plugins/Enterprise/DataprepEditor/Source/DataprepCore/Private/InterchangeFileProducer.cpp b/Engine/Plugins/Enterprise/DataprepEditor/Source/DataprepCore/Private/InterchangeFileProducer.cpp
--- a/Engine/Plugins/Enterprise/DataprepEditor/Source/DataprepCore/Private/InterchangeFileProducer.cpp
+++ b/Engine/Plugins/Enterprise/DataprepEditor/Source/DataprepCore/Private/InterchangeFileProducer.cpp
@@ -39,9 +39,6 @@
 const FText InterchangeFileProducerLabel( LOCTEXT( "InterchangeFileProducerLabel", "Interchange file importer" ) );
 const FText InterchangeFileProducerDescription( LOCTEXT( "InterchangeFileProducerDesc", "Reads a file supported by Interchange and its dependent assets" ) );
 
-const TCHAR* WildCharCriteria = TEXT( "*.*" );
-const TCHAR* ExtensionPrefix = TEXT( "*." );
-
 bool bGEnableInterchangeProducer = false;
 FAutoConsoleVariableRef GEnableInterchangeProducerVar(
 	TEXT("Dataprep.Interchange.Producer"),
@@ -417,7 +414,7 @@ private:
 	TSharedPtr< SEditableText > FileName;
 };
 
-void FDataprepContentProducerDetails::CustomizeDetails(IDetailLayoutBuilder& DetailBuilder)
+void FInterchangeContentProducerDetails::CustomizeDetails(IDetailLayoutBuilder& DetailBuilder)
 {
 	TArray< TWeakObjectPtr< UObject > > Objects;
 	DetailBuilder.GetObjectsBeingCustomized(Objects);
@@ -440,19 +437,19 @@ void FDataprepContentProducerDetails::CustomizeDetails(IDetailLayoutBuilder& Det
 	}
 }
 
-FSlateColor FDataprepContentProducerDetails::GetStatusColorAndOpacity() const
+FSlateColor FInterchangeContentProducerDetails::GetStatusColorAndOpacity() const
 {
 	return  IsProducerSuperseded() ? FLinearColor::Red : FAppStyle::Get().GetSlateColor("DefaultForeground");
 }
 
-bool FDataprepContentProducerDetails::IsProducerSuperseded() const
+bool FInterchangeContentProducerDetails::IsProducerSuperseded() const
 {
 	return ProducerIndex != INDEX_NONE ? AssetProducers->IsProducerSuperseded(ProducerIndex) : false;
 }
 
 void FInterchangeFileProducerDetails::CustomizeDetails(IDetailLayoutBuilder& DetailBuilder)
 {
-	FDataprepContentProducerDetails::CustomizeDetails(DetailBuilder);
+	FInterchangeContentProducerDetails::CustomizeDetails(DetailBuilder);
 
 	UInterchangeFileProducer* FileProducer = Cast< UInterchangeFileProducer >(Producer);
 	check( FileProducer );
@@ -510,7 +507,7 @@ void FInterchangeFileProducerDetails::CustomizeDetails(IDetailLayoutBuilder& Det
 				{
 					return this->IsProducerSuperseded() ? FEditorFontGlyphs::Exclamation_Triangle : FEditorFontGlyphs::File;
 				}))
-				.ColorAndOpacity(this, &FDataprepContentProducerDetails::GetStatusColorAndOpacity)
+				.ColorAndOpacity(this, &FInterchangeContentProducerDetails::GetStatusColorAndOpacity)
 		]
 	];
 
diff --git a/Engine/Plugins/Enterprise/DataprepEditor/Source/DataprepCore/Private/InterchangeFileProducer.h b/Engine/Plugins/Enterprise/DataprepEditor/Source/DataprepCore/Private/InterchangeFileProducer.h
--- a/Engine/Plugins/Enterprise/DataprepEditor/Source/DataprepCore/Private/InterchangeFileProducer.h
+++ b/Engine/Plugins/Enterprise/DataprepEditor/Source/DataprepCore/Private/InterchangeFileProducer.h
@@ -72,7 +72,7 @@ private:
 	friend class SInterchangeFileProducerFileProperty;
 };
 
-class FDataprepContentProducerDetails : public IDetailCustomization
+class FInterchangeContentProducerDetails : public IDetailCustomization
 {
 public:
 	virtual void CustomizeDetails( IDetailLayoutBuilder& DetailBuilder ) override;
@@ -86,7 +86,7 @@ protected:
 };
 
 // Customization of the details of the Interchange producer for the data prep editor.
-class DATAPREPCORE_API FInterchangeFileProducerDetails : public FDataprepContentProducerDetails
+class DATAPREPCORE_API FInterchangeFileProducerDetails : public FInterchangeContentProducerDetails
 {
 public:
 	static TSharedRef< IDetailCustomization > MakeDetails() { return MakeShared<FInterchangeFileProducerDetails>(); };
