Improved UDataprepMergeActorsOperation::MergeStaticMeshActors based on user feedback.

#rnx

diff --git a/Engine/Plugins/Enterprise/DataprepEditor/Source/DataprepCore/Private/DataprepCorePrivateUtils.cpp b/Engine/Plugins/Enterprise/DataprepEditor/Source/DataprepCore/Private/DataprepCorePrivateUtils.cpp
--- a/Engine/Plugins/Enterprise/DataprepEditor/Source/DataprepCore/Private/DataprepCorePrivateUtils.cpp
+++ b/Engine/Plugins/Enterprise/DataprepEditor/Source/DataprepCore/Private/DataprepCorePrivateUtils.cpp
@@ -45,7 +45,7 @@ const FString& DataprepCorePrivateUtils::GetRootTemporaryDir()
 
 const FString& DataprepCorePrivateUtils::GetRootPackagePath()
 {
-	static FString RootPackagePath( TEXT("/Engine/DataprepCore/Transient") );
+	static FString RootPackagePath( TEXT("/DataprepCore/Transient") );
 	return RootPackagePath;
 }
 
diff --git a/Engine/Plugins/Enterprise/DataprepEditor/Source/DataprepLibraries/DataprepLibraries.Build.cs b/Engine/Plugins/Enterprise/DataprepEditor/Source/DataprepLibraries/DataprepLibraries.Build.cs
--- a/Engine/Plugins/Enterprise/DataprepEditor/Source/DataprepLibraries/DataprepLibraries.Build.cs
+++ b/Engine/Plugins/Enterprise/DataprepEditor/Source/DataprepLibraries/DataprepLibraries.Build.cs
@@ -15,6 +15,7 @@ namespace UnrealBuildTool.Rules
 					"DataprepCore",
 					"DataprepEditor",
 					"DatasmithContent",
+					"DerivedDataCache",
 					"EditorFramework",
 					"EditorScriptingUtilities",
 					"Engine",
diff --git a/Engine/Plugins/Enterprise/DataprepEditor/Source/DataprepLibraries/Private/DataprepEditingOperations.cpp b/Engine/Plugins/Enterprise/DataprepEditor/Source/DataprepLibraries/Private/DataprepEditingOperations.cpp
--- a/Engine/Plugins/Enterprise/DataprepEditor/Source/DataprepLibraries/Private/DataprepEditingOperations.cpp
+++ b/Engine/Plugins/Enterprise/DataprepEditor/Source/DataprepLibraries/Private/DataprepEditingOperations.cpp
@@ -12,6 +12,7 @@
 #include "Engine/Blueprint.h"
 #include "IContentBrowserSingleton.h"
 #include "ContentBrowserModule.h"
+#include "DerivedDataCacheInterface.h"
 #include "Editor.h"
 #include "Engine/Level.h"
 #include "Engine/StaticMesh.h"
@@ -1055,9 +1056,17 @@ namespace DatasmithEditingOperationsUtils
 
 		TArray<UObject*> ObjectsToDelete;
 		ObjectsToDelete.Reserve(InMergedComponents.Num());
+		TSet<UStaticMesh*> StaticMeshes;
+		StaticMeshes.Reserve(InMergedComponents.Num());
 
 		for( UStaticMeshComponent* Component : InMergedComponents )
 		{
+			if (UStaticMesh* StaticMesh = Component->GetStaticMesh())
+			{
+				Component->SetStaticMesh(nullptr);
+				StaticMeshes.Add(StaticMesh);
+			}
+
 			if( Component->GetNumChildrenComponents() > 0 )
 			{
 				// Reparent children
@@ -1105,6 +1114,34 @@ namespace DatasmithEditingOperationsUtils
 			}
 		}
 
+		// Remove from set any static mesh still used by another component
+		for (TObjectIterator<UStaticMeshComponent> It; It; ++It)
+		{
+			if (UStaticMesh* StaticMesh = It->GetStaticMesh())
+			{
+				if (StaticMeshes.Contains(StaticMesh))
+				{
+					StaticMeshes.Remove(StaticMesh);
+				}
+			}
+		}
+
+		if (StaticMeshes.Num() > 0)
+		{
+			FDerivedDataCacheInterface& DDC = GetDerivedDataCacheRef();
+
+			for (UStaticMesh* StaticMesh : StaticMeshes)
+			{
+				if (FStaticMeshRenderData* RenderData = StaticMesh->GetRenderData())
+				{
+					// Mark this mesh render data as transient in order to delete any file written in the DDC
+					DDC.MarkTransient(*RenderData->DerivedDataKey);
+				}
+			}
+
+			ObjectsToDelete.Append(StaticMeshes.Array());
+		}
+
 		return MoveTemp( ObjectsToDelete );
 	}
 
