Addressed issues reported on the Interchange file producer

#virtualized

diff --git a/Engine/Plugins/Enterprise/DataprepEditor/Source/DataprepCore/DataprepCore.Build.cs b/Engine/Plugins/Enterprise/DataprepEditor/Source/DataprepCore/DataprepCore.Build.cs
--- a/Engine/Plugins/Enterprise/DataprepEditor/Source/DataprepCore/DataprepCore.Build.cs
+++ b/Engine/Plugins/Enterprise/DataprepEditor/Source/DataprepCore/DataprepCore.Build.cs
@@ -13,7 +13,7 @@ namespace UnrealBuildTool.Rules
 					"BlueprintGraph",
 					"Core",
 					"CoreUObject",
-					"DesktopPlatform",
+					"DatasmithContent",
 					"EditorFramework",
 					"EditorScriptingUtilities",
 					"EditorStyle",
diff --git a/Engine/Plugins/Enterprise/DataprepEditor/Source/DataprepCore/Private/InterchangeDataprepPipeline.cpp b/Engine/Plugins/Enterprise/DataprepEditor/Source/DataprepCore/Private/InterchangeDataprepPipeline.cpp
--- a/Engine/Plugins/Enterprise/DataprepEditor/Source/DataprepCore/Private/InterchangeDataprepPipeline.cpp
+++ b/Engine/Plugins/Enterprise/DataprepEditor/Source/DataprepCore/Private/InterchangeDataprepPipeline.cpp
@@ -2,13 +2,13 @@
 
 #include "InterchangeDataprepPipeline.h"
 
-//#include "Nodes/InterchangeFactoryBaseNode.h"
 #include "InterchangeActorFactoryNode.h"
-
-#include "Engine/Texture.h"
-#include "Materials/MaterialInterface.h"
-#include "Engine/StaticMesh.h"
-#include "LevelSequence.h"
+#include "InterchangeLevelSequenceFactoryNode.h"
+#include "InterchangeMaterialFactoryNode.h"
+#include "InterchangeSceneVariantSetsFactoryNode.h"
+#include "InterchangeStaticMeshFactoryNode.h"
+#include "InterchangeTextureFactoryNode.h"
+#include "Nodes/InterchangeFactoryBaseNode.h"
 
 void UInterchangeDataprepLevelPipeline::ExecutePipeline(UInterchangeBaseNodeContainer* NodeContainer, const TArray<UInterchangeSourceData*>& InSourceDatas, const FString& ContentBasePath)
 {
@@ -19,35 +19,48 @@ void UInterchangeDataprepLevelPipeline::ExecutePipeline(UInterchangeBaseNodeCont
 
 	Super::ExecutePipeline(NodeContainer, InSourceDatas, ContentBasePath);
 
-	static TMap<UClass*, FString> SubPathsPerClass{
-		{ UTexture::StaticClass(), "Textures" },
-		{ UMaterialInterface::StaticClass(), "Materials" },
-		{ UStaticMesh::StaticClass(), "Geometries" },
-		{ ULevelSequence::StaticClass(), TEXT("Animations") },
-	};
-
 	// Compute unique prefix based on file path
-	FString UniquePrefix = FString::FromInt(GetTypeHash(InSourceDatas[0]->GetFilename()));
+	const FString PackageSubPath = FPaths::GetBaseFilename(InSourceDatas[0]->GetFilename());
 
-	auto UpdateFactoryNodes = [this,&UniquePrefix](const FString& NodeUid, UInterchangeFactoryBaseNode* FactoryNode)
+	auto UpdateFactoryNodes = [this,&PackageSubPath](const FString& NodeUid, UInterchangeFactoryBaseNode* FactoryNode)
 	{
-		for (TPair<UClass*, FString>& Entry : SubPathsPerClass)
+		// Append prefix to prevent actor's name collision in Dataprep world
+		if (UInterchangeActorFactoryNode* ActorFactoryNode = Cast< UInterchangeActorFactoryNode>(FactoryNode))
+		{
+			FString NewLabel = PackageSubPath + TEXT("_") + ActorFactoryNode->GetDisplayLabel();
+			ActorFactoryNode->SetDisplayLabel(NewLabel);
+		}
+		else if (UInterchangeTextureFactoryNode* TextureFactoryNode = Cast<UInterchangeTextureFactoryNode>(FactoryNode))
+		{
+			TextureFactoryNode->SetCustomSubPath(FPaths::Combine(PackageSubPath, "Textures"));
+		}
+		else if (UInterchangeBaseMaterialFactoryNode* MaterialFactoryNode = Cast<UInterchangeBaseMaterialFactoryNode>(FactoryNode))
 		{
-			if (FactoryNode->GetObjectClass() && FactoryNode->GetObjectClass()->IsChildOf(Entry.Key))
+			if (MaterialFactoryNode->IsA<UInterchangeMaterialFactoryNode>())
+			{
+				MaterialFactoryNode->SetCustomSubPath(FPaths::Combine(PackageSubPath, "Materials/References"));
+			}
+			else if (MaterialFactoryNode->IsA<UInterchangeMaterialFunctionFactoryNode>())
 			{
-				FactoryNode->SetCustomSubPath(Entry.Value);
-				break;
+				MaterialFactoryNode->SetCustomSubPath(FPaths::Combine(PackageSubPath, "Materials/References/Functions"));
+			}
+			else
+			{
+				MaterialFactoryNode->SetCustomSubPath(FPaths::Combine(PackageSubPath, "Materials"));
 			}
 		}
-
-		// Append prefix to prevent actor's name collision in Dataprep world
-		if (UInterchangeActorFactoryNode* ActorFactoryNode = Cast< UInterchangeActorFactoryNode>(FactoryNode))
+		else if (UInterchangeStaticMeshFactoryNode* MeshFactoryNode = Cast<UInterchangeStaticMeshFactoryNode>(FactoryNode))
 		{
-			FString NewLabel = UniquePrefix + TEXT("_") + ActorFactoryNode->GetDisplayLabel();
-			ActorFactoryNode->SetDisplayLabel(NewLabel);
+			MeshFactoryNode->SetCustomSubPath(FPaths::Combine(PackageSubPath, "Geometries"));
+		}
+		else if (UInterchangeLevelSequenceFactoryNode* SequenceFactoryNode = Cast<UInterchangeLevelSequenceFactoryNode>(FactoryNode))
+		{
+			SequenceFactoryNode->SetCustomSubPath(FPaths::Combine(PackageSubPath, "Animations"));
+		}
+		else if (UInterchangeSceneVariantSetsFactoryNode* VariantFactoryNode = Cast<UInterchangeSceneVariantSetsFactoryNode>(FactoryNode))
+		{
+			VariantFactoryNode->SetCustomSubPath(FPaths::Combine(PackageSubPath, "Variants"));
 		}
-
-		//FactoryNode->SetEnabled(true);
 	};
 
 	//Find all factory node we need for this pipeline
diff --git a/Engine/Plugins/Enterprise/DataprepEditor/Source/DataprepCore/Private/InterchangeFileProducer.cpp b/Engine/Plugins/Enterprise/DataprepEditor/Source/DataprepCore/Private/InterchangeFileProducer.cpp
--- a/Engine/Plugins/Enterprise/DataprepEditor/Source/DataprepCore/Private/InterchangeFileProducer.cpp
+++ b/Engine/Plugins/Enterprise/DataprepEditor/Source/DataprepCore/Private/InterchangeFileProducer.cpp
@@ -7,6 +7,7 @@
 #include "IDataprepProgressReporter.h"
 #include "InterchangeDataprepPipeline.h"
 
+#include "InterchangeProjectSettings.h"
 #include "InterchangeManager.h"
 
 #include "AssetRegistry/AssetRegistryModule.h"
@@ -47,8 +48,6 @@ FAutoConsoleVariableRef GEnableInterchangeProducerVar(
 	ECVF_Default
 );
 
-TArray<FSoftObjectPath> UInterchangeFileProducer::OverridePipelines;
-
 namespace FInterchangeFileProducerUtils
 {
 	/** Delete all the assets stored under the specified path */
@@ -133,17 +132,42 @@ namespace FInterchangeFileProducerUtils
 		TArray<UObject*> ImportOptions;
 		TWeakPtr< SWindow > Window;
 	};
+
+	void GetPipelinesFromSourceData(UInterchangeSourceData& SourceData, TArray<FSoftObjectPath>& OutPipelines)
+	{
+		const FInterchangeImportSettings& InterchangeImportSettings = FInterchangeProjectSettingsUtils::GetDefaultImportSettings(true);
+		const TMap<FName, FInterchangePipelineStack>& DefaultPipelineStacks = InterchangeImportSettings.PipelineStacks;
+		FName DefaultStackName = FInterchangeProjectSettingsUtils::GetDefaultPipelineStackName(true, SourceData);
+		
+		OutPipelines.Empty();
+
+		UE::Interchange::FScopedTranslator ScopedTranslator(&SourceData);
+
+		for (const TPair<FName, FInterchangePipelineStack>& PipelineStackInfo : DefaultPipelineStacks)
+		{
+			if (PipelineStackInfo.Key == DefaultStackName)
+			{
+				const FInterchangePipelineStack& PipelineStack = PipelineStackInfo.Value;
+				OutPipelines = PipelineStack.Pipelines;
+
+				for (const FInterchangeTranslatorPipelines& TranslatorPipelines : PipelineStack.PerTranslatorPipelines)
+				{
+					const UClass* TranslatorClass = TranslatorPipelines.Translator.LoadSynchronous();
+					if (ScopedTranslator.GetTranslator()->IsA(TranslatorClass))
+					{
+						OutPipelines = TranslatorPipelines.Pipelines;
+						break;
+					}
+				}
+
+				break;
+			}
+		}
+	}
 }
 
 UInterchangeFileProducer::UInterchangeFileProducer()
 {
-	if (HasAnyFlags(RF_ClassDefaultObject))
-	{
-		// #dataprep_todo Allow users to override the stack of pipelines
-		OverridePipelines.Add(FSoftObjectPath(TEXT("/Interchange/Pipelines/DefaultSceneAssetsPipeline.DefaultSceneAssetsPipeline")));
-		OverridePipelines.Add(FSoftObjectPath(TEXT("/Interchange/Pipelines/DefaultSceneLevelPipeline.DefaultSceneLevelPipeline")));
-		OverridePipelines.Add(FSoftObjectPath(TEXT("/DataprepEditor/DefaultDataprepPipeline.DefaultDataprepPipeline")));
-	}
 }
 
 bool UInterchangeFileProducer::IsActive()
@@ -192,6 +216,7 @@ bool UInterchangeFileProducer::InitTranslator()
 
 bool UInterchangeFileProducer::Execute(TArray< TWeakObjectPtr< UObject > >& OutAssets)
 {
+	using namespace FInterchangeFileProducerUtils;
 	using namespace UE::Interchange;
 	using FInterchangeResults = TTuple<FAssetImportResultRef, FSceneImportResultRef>;
 
@@ -206,18 +231,22 @@ bool UInterchangeFileProducer::Execute(TArray< TWeakObjectPtr< UObject > >& OutA
 
 	FScopedSourceData ScopedSourceData(FilePath);
 
-	if (!InterchangeManager.CanTranslateSourceData(ScopedSourceData.GetSourceData()))
+	UInterchangeSourceData* SourceData = ScopedSourceData.GetSourceData();
+
+	if (!InterchangeManager.CanTranslateSourceData(SourceData))
 	{
 		return false;
 	}
-	
+
 	FImportAssetParameters ImportAssetParameters;
 	
 	ImportAssetParameters.bIsAutomated = true;
 	ImportAssetParameters.ImportLevel = Context.WorldPtr.Get()->GetCurrentLevel();
-	ImportAssetParameters.OverridePipelines = OverridePipelines;
 
-	FInterchangeResults ImportResults = InterchangeManager.ImportSceneAsync(TransientPackage->GetPathName(), ScopedSourceData.GetSourceData(), ImportAssetParameters);
+	GetPipelinesFromSourceData(*SourceData, ImportAssetParameters.OverridePipelines);
+	ImportAssetParameters.OverridePipelines.Add(FSoftObjectPath(TEXT("/DataprepEditor/DefaultDataprepPipeline.DefaultDataprepPipeline")));
+
+	FInterchangeResults ImportResults = InterchangeManager.ImportSceneAsync(TransientPackage->GetPathName(), SourceData, ImportAssetParameters);
 
 	ImportResults.Get<0>()->WaitUntilDone();
 	ImportResults.Get<1>()->WaitUntilDone();
diff --git a/Engine/Plugins/Enterprise/DataprepEditor/Source/DataprepCore/Private/InterchangeFileProducer.h b/Engine/Plugins/Enterprise/DataprepEditor/Source/DataprepCore/Private/InterchangeFileProducer.h
--- a/Engine/Plugins/Enterprise/DataprepEditor/Source/DataprepCore/Private/InterchangeFileProducer.h
+++ b/Engine/Plugins/Enterprise/DataprepEditor/Source/DataprepCore/Private/InterchangeFileProducer.h
@@ -65,8 +65,6 @@ private:
 	UPROPERTY( Transient, DuplicateTransient )
 	TObjectPtr<UPackage> TransientPackage = nullptr;
 
-	static TArray<FSoftObjectPath> OverridePipelines;
-
 	TArray< TWeakObjectPtr< UObject > > Assets;
 
 	friend class SInterchangeFileProducerFileProperty;
