Fixed Dataprep crash when deleting a producer
Fixed flawed logic with UDataprepAssetProducers member when loading an existing Dataprep asset
Skipping the creation of the InterchangeSceneImportAsset for the time being since the consumer is a Datasmith one

diff --git a/Engine/Plugins/Enterprise/DataprepEditor/Source/DataprepCore/Private/DataprepAssetInterface.cpp b/Engine/Plugins/Enterprise/DataprepEditor/Source/DataprepCore/Private/DataprepAssetInterface.cpp
--- a/Engine/Plugins/Enterprise/DataprepEditor/Source/DataprepCore/Private/DataprepAssetInterface.cpp
+++ b/Engine/Plugins/Enterprise/DataprepEditor/Source/DataprepCore/Private/DataprepAssetInterface.cpp
@@ -39,10 +39,11 @@ void UDataprepAssetInterface::PostLoad()
 
 		FAssetRegistryModule::AssetCreated( Inputs );
 		Inputs->MarkPackageDirty();
-
-		Inputs->GetOnChanged().AddUObject( this, &UDataprepAssetInterface::OnProducersChanged );
 	}
 
+	ensure(Inputs);
+	Inputs->GetOnChanged().AddUObject(this, &UDataprepAssetInterface::OnProducersChanged);
+
 	if(Output != nullptr)
 	{
 		Output->GetOnChanged().AddUObject( this, &UDataprepAssetInterface::OnConsumerChanged );
diff --git a/Engine/Plugins/Enterprise/DataprepEditor/Source/DataprepCore/Private/DataprepAssetProducers.cpp b/Engine/Plugins/Enterprise/DataprepEditor/Source/DataprepCore/Private/DataprepAssetProducers.cpp
--- a/Engine/Plugins/Enterprise/DataprepEditor/Source/DataprepCore/Private/DataprepAssetProducers.cpp
+++ b/Engine/Plugins/Enterprise/DataprepEditor/Source/DataprepCore/Private/DataprepAssetProducers.cpp
@@ -8,6 +8,7 @@
 #include "Shared/DataprepCorePrivateUtils.h"
 #include "DataprepCoreUtils.h"
 
+#include "ScopedTransaction.h"
 
 
 #define LOCTEXT_NAMESPACE "DataprepAssetProducers"
@@ -191,51 +192,54 @@ bool UDataprepAssetProducers::RemoveProducer(int32 IndexToRemove)
 
 		// Array of producers superseded by removed producer
 		TArray<int32> ProducersToRevisit;
-		ProducersToRevisit.Reserve( AssetProducers.Num() );
+		ProducersToRevisit.Reserve(AssetProducers.Num());
+
 		{
+			const FScopedTransaction Transaction(LOCTEXT("Producers_RemoveProducer", "Remove Producer"));
+
 			Modify();
 
-			if(UDataprepContentProducer* Producer = AssetProducers[IndexToRemove].Producer)
+			if (UDataprepContentProducer* Producer = AssetProducers[IndexToRemove].Producer)
 			{
-				Producer->GetOnChanged().RemoveAll( this );
+				Producer->GetOnChanged().RemoveAll(this);
 
 				Producer->PreEditChange(nullptr);
 				Producer->Modify();
 
-				DataprepCorePrivateUtils::DeleteRegisteredAsset( Producer );
+				DataprepCorePrivateUtils::DeleteRegisteredAsset(Producer);
 
 				Producer->PostEditChange();
 			}
 
-			AssetProducers.RemoveAt( IndexToRemove );
+			AssetProducers.RemoveAt(IndexToRemove);
 
-			if( AssetProducers.Num() == 1 )
+			if (AssetProducers.Num() == 1)
 			{
 				AssetProducers[0].SupersededBy = INDEX_NONE;
 			}
-			else if(AssetProducers.Num() > 1)
+			else if (AssetProducers.Num() > 1)
 			{
 				// Update value stored in SupersededBy property where applicable
-				for( int32 Index = 0; Index < AssetProducers.Num(); ++Index )
+				for (int32 Index = 0; Index < AssetProducers.Num(); ++Index)
 				{
 					FDataprepAssetProducer& AssetProducer = AssetProducers[Index];
 
-					if( AssetProducer.SupersededBy == IndexToRemove )
+					if (AssetProducer.SupersededBy == IndexToRemove)
 					{
 						AssetProducer.SupersededBy = INDEX_NONE;
-						ProducersToRevisit.Add( Index );
+						ProducersToRevisit.Add(Index);
 					}
-					else if( AssetProducer.SupersededBy > IndexToRemove )
+					else if (AssetProducer.SupersededBy > IndexToRemove)
 					{
 						--AssetProducer.SupersededBy;
 					}
 				}
 			}
 
-			for( int32 ProducerIndex : ProducersToRevisit )
+			for (int32 ProducerIndex : ProducersToRevisit)
 			{
 				bool bLocalChangeAll = false;
-				ValidateProducerChanges( ProducerIndex, bLocalChangeAll );
+				ValidateProducerChanges(ProducerIndex, bLocalChangeAll);
 				bChangeAll |= bLocalChangeAll;
 			}
 		}
diff --git a/Engine/Plugins/Enterprise/DataprepEditor/Source/DataprepCore/Private/InterchangeDataprepPipeline.cpp b/Engine/Plugins/Enterprise/DataprepEditor/Source/DataprepCore/Private/InterchangeDataprepPipeline.cpp
--- a/Engine/Plugins/Enterprise/DataprepEditor/Source/DataprepCore/Private/InterchangeDataprepPipeline.cpp
+++ b/Engine/Plugins/Enterprise/DataprepEditor/Source/DataprepCore/Private/InterchangeDataprepPipeline.cpp
@@ -5,6 +5,7 @@
 #include "InterchangeActorFactoryNode.h"
 #include "InterchangeLevelSequenceFactoryNode.h"
 #include "InterchangeMaterialFactoryNode.h"
+#include "InterchangeSceneImportAssetFactoryNode.h"
 #include "InterchangeSceneVariantSetsFactoryNode.h"
 #include "InterchangeStaticMeshFactoryNode.h"
 #include "InterchangeTextureFactoryNode.h"
@@ -61,6 +62,12 @@ void UInterchangeDataprepLevelPipeline::ExecutePipeline(UInterchangeBaseNodeCont
 		{
 			VariantFactoryNode->SetCustomSubPath(FPaths::Combine(PackageSubPath, "Variants"));
 		}
+		// Remove UInterchangeSceneImportAssetFactoryNode from the import
+		// #dataprep_todo : Revisit when consumer will have been updated.
+		else if (UInterchangeSceneImportAssetFactoryNode* SceneImportFactoryNode = Cast<UInterchangeSceneImportAssetFactoryNode>(FactoryNode))
+		{
+			SceneImportFactoryNode->SetEnabled(false);
+		}
 	};
 
 	//Find all factory node we need for this pipeline
diff --git a/Engine/Plugins/Enterprise/DataprepEditor/Source/DataprepEditor/Private/DataprepEditor.cpp b/Engine/Plugins/Enterprise/DataprepEditor/Source/DataprepEditor/Private/DataprepEditor.cpp
--- a/Engine/Plugins/Enterprise/DataprepEditor/Source/DataprepEditor/Private/DataprepEditor.cpp
+++ b/Engine/Plugins/Enterprise/DataprepEditor/Source/DataprepEditor/Private/DataprepEditor.cpp
@@ -792,6 +792,12 @@ void FDataprepEditor::CleanPreviewWorld()
 			UPackage* Package = Cast<UPackage>(PackagePath.ResolveObject());
 
 			UObject* ObjectToDelete = StaticFindObjectFast(nullptr, Package, *SoftObjectPath.GetAssetName());
+			if (!ObjectToDelete)
+			{
+				continue;
+			}
+
+			ObjectToDelete->ClearFlags(RF_Transactional);
 
 			if (UMaterialInstanceConstant* MaterialInstanceConstant = Cast<UMaterialInstanceConstant>(ObjectToDelete))
 			{
diff --git a/Engine/Plugins/Enterprise/DataprepEditor/Source/DataprepEditor/Private/Widgets/SDataprepProducersWidget.cpp b/Engine/Plugins/Enterprise/DataprepEditor/Source/DataprepEditor/Private/Widgets/SDataprepProducersWidget.cpp
--- a/Engine/Plugins/Enterprise/DataprepEditor/Source/DataprepEditor/Private/Widgets/SDataprepProducersWidget.cpp
+++ b/Engine/Plugins/Enterprise/DataprepEditor/Source/DataprepEditor/Private/Widgets/SDataprepProducersWidget.cpp
@@ -55,7 +55,6 @@ void FContentProducerEntry::RemoveProducer()
 {
 	if( UDataprepAssetProducers* AssetProducers = AssetProducersPtr.Get() )
 	{
-		const FScopedTransaction Transaction( LOCTEXT("Producers_RemoveProducer", "Remove Producer") );
 		AssetProducers->RemoveProducer( ProducerIndex );
 	}
 }
