Prevent clang analysis warnings for uninitialized data from FArchive derived readers

#rnx

diff --git a/Engine/Plugins/Enterprise/DataprepEditor/Source/DataprepEditor/Private/DataprepSnapshot.cpp b/Engine/Plugins/Enterprise/DataprepEditor/Source/DataprepEditor/Private/DataprepSnapshot.cpp
--- a/Engine/Plugins/Enterprise/DataprepEditor/Source/DataprepEditor/Private/DataprepSnapshot.cpp
+++ b/Engine/Plugins/Enterprise/DataprepEditor/Source/DataprepEditor/Private/DataprepSnapshot.cpp
@@ -121,6 +121,10 @@ namespace DataprepSnapshotUtil
 
 		void Serialize( void* Data, int64 Num )
 		{
+#if defined(__clang_analyzer__)
+			// Suppress uninitialized data static analysis warning
+			FMemory::Memzero(Data, Num);
+#endif
 			if (Num && !ArIsError)
 			{
 				// Only serialize if we have the requested amount of data
diff --git a/Engine/Source/Runtime/Core/Private/Serialization/LargeMemoryReader.cpp b/Engine/Source/Runtime/Core/Private/Serialization/LargeMemoryReader.cpp
--- a/Engine/Source/Runtime/Core/Private/Serialization/LargeMemoryReader.cpp
+++ b/Engine/Source/Runtime/Core/Private/Serialization/LargeMemoryReader.cpp
@@ -26,6 +26,10 @@ FLargeMemoryReader::FLargeMemoryReader(const uint8* InData, const int64 Num, ELa
 
 void FLargeMemoryReader::Serialize(void* OutData, int64 Num)
 {
+#if defined(__clang_analyzer__)
+	// Suppress uninitialized data static analysis warning
+	FMemory::Memzero(Data, Num);
+#endif
 	if (Num && !IsError())
 	{
 		// Only serialize if we have the requested amount of data
diff --git a/Engine/Source/Runtime/Core/Public/Serialization/MemoryReader.h b/Engine/Source/Runtime/Core/Public/Serialization/MemoryReader.h
--- a/Engine/Source/Runtime/Core/Public/Serialization/MemoryReader.h
+++ b/Engine/Source/Runtime/Core/Public/Serialization/MemoryReader.h
@@ -34,6 +34,10 @@ public:
 
 	void Serialize( void* Data, int64 Num )
 	{
+#if defined(__clang_analyzer__)
+		// Suppress uninitialized data static analysis warning
+		FMemory::Memzero(Data, Num);
+#endif
 		if (Num && !IsError())
 		{
 			// Only serialize if we have the requested amount of data
diff --git a/Engine/Source/Runtime/Core/Public/Serialization/StaticMemoryReader.h b/Engine/Source/Runtime/Core/Public/Serialization/StaticMemoryReader.h
--- a/Engine/Source/Runtime/Core/Public/Serialization/StaticMemoryReader.h
+++ b/Engine/Source/Runtime/Core/Public/Serialization/StaticMemoryReader.h
@@ -32,6 +32,10 @@ public:
 
 	void Serialize( void* Data, int64 Num )
 	{
+#if defined(__clang_analyzer__)
+		// Suppress uninitialized data static analysis warning
+		FMemory::Memzero(Data, Num);
+#endif
 		if (Num && !IsError())
 		{
 			// Only serialize if we have the requested amount of data
diff --git a/Engine/Source/Runtime/CoreUObject/Public/Serialization/ObjectReader.h b/Engine/Source/Runtime/CoreUObject/Public/Serialization/ObjectReader.h
--- a/Engine/Source/Runtime/CoreUObject/Public/Serialization/ObjectReader.h
+++ b/Engine/Source/Runtime/CoreUObject/Public/Serialization/ObjectReader.h
@@ -55,6 +55,10 @@ public:
 
 	void Serialize(void* Data, int64 Num)
 	{
+#if defined(__clang_analyzer__)
+		// Suppress uninitialized data static analysis warning
+		FMemory::Memzero(Data, Num);
+#endif
 		if (Num && !IsError())
 		{
 			// Only serialize if we have the requested amount of data
diff --git a/Engine/Source/Runtime/Net/Core/Private/Net/Core/Misc/PacketAudit.cpp b/Engine/Source/Runtime/Net/Core/Private/Net/Core/Misc/PacketAudit.cpp
--- a/Engine/Source/Runtime/Net/Core/Private/Net/Core/Misc/PacketAudit.cpp
+++ b/Engine/Source/Runtime/Net/Core/Private/Net/Core/Misc/PacketAudit.cpp
@@ -157,6 +157,10 @@ public:
 
 	virtual void Serialize(void* Data, int64 Num) override
 	{
+#if defined(__clang_analyzer__)
+		// Suppress uninitialized data static analysis warning
+		FMemory::Memzero(Data, Num);
+#endif
 		if (Num > 0 && !IsError())
 		{
 			if (Offset + Num <= TotalSize())
diff --git a/Engine/Source/Runtime/Online/BuildPatchServices/Private/Data/ManifestUObject.cpp b/Engine/Source/Runtime/Online/BuildPatchServices/Private/Data/ManifestUObject.cpp
--- a/Engine/Source/Runtime/Online/BuildPatchServices/Private/Data/ManifestUObject.cpp
+++ b/Engine/Source/Runtime/Online/BuildPatchServices/Private/Data/ManifestUObject.cpp
@@ -356,6 +356,10 @@ public:
 
 	virtual void Serialize(void* Data, int64 Num) override
 	{
+#if defined(__clang_analyzer__)
+		// Suppress uninitialized data static analysis warning
+		FMemory::Memzero(Data, Num);
+#endif
 		if (Num && !IsError())
 		{
 			// Only serialize if we have the requested amount of data
diff --git a/Engine/Source/Runtime/Online/BuildPatchServices/Private/Tests/Fake/FileSystem.fake.h b/Engine/Source/Runtime/Online/BuildPatchServices/Private/Tests/Fake/FileSystem.fake.h
--- a/Engine/Source/Runtime/Online/BuildPatchServices/Private/Tests/Fake/FileSystem.fake.h
+++ b/Engine/Source/Runtime/Online/BuildPatchServices/Private/Tests/Fake/FileSystem.fake.h
@@ -36,6 +36,10 @@ namespace BuildPatchServices
 
 		void Serialize(void* Data, int64 Num)
 		{
+#if defined(__clang_analyzer__)
+			// Suppress uninitialized data static analysis warning
+			FMemory::Memzero(Data, Num);
+#endif
 			if (Num && !IsError())
 			{
 				if (Offset + Num <= TotalSize())
