[Backout] - CL32838690
Original CL Desc
-----------------------------------------------------------------
Prevent clang analysis warnings for uninitialized data from FArchive derived readers

#rnx

diff --git a/Engine/Plugins/Enterprise/DataprepEditor/Source/DataprepEditor/Private/DataprepSnapshot.cpp b/Engine/Plugins/Enterprise/DataprepEditor/Source/DataprepEditor/Private/DataprepSnapshot.cpp
--- a/Engine/Plugins/Enterprise/DataprepEditor/Source/DataprepEditor/Private/DataprepSnapshot.cpp
+++ b/Engine/Plugins/Enterprise/DataprepEditor/Source/DataprepEditor/Private/DataprepSnapshot.cpp
@@ -121,10 +121,6 @@ namespace DataprepSnapshotUtil
 
 		void Serialize( void* Data, int64 Num )
 		{
-#if defined(__clang_analyzer__)
-			// Suppress uninitialized data static analysis warning
-			FMemory::Memzero(Data, Num);
-#endif
 			if (Num && !ArIsError)
 			{
 				// Only serialize if we have the requested amount of data
diff --git a/Engine/Source/Runtime/Core/Private/Serialization/LargeMemoryReader.cpp b/Engine/Source/Runtime/Core/Private/Serialization/LargeMemoryReader.cpp
--- a/Engine/Source/Runtime/Core/Private/Serialization/LargeMemoryReader.cpp
+++ b/Engine/Source/Runtime/Core/Private/Serialization/LargeMemoryReader.cpp
@@ -26,10 +26,6 @@ FLargeMemoryReader::FLargeMemoryReader(const uint8* InData, const int64 Num, ELa
 
 void FLargeMemoryReader::Serialize(void* OutData, int64 Num)
 {
-#if defined(__clang_analyzer__)
-	// Suppress uninitialized data static analysis warning
-	FMemory::Memzero(Data, Num);
-#endif
 	if (Num && !IsError())
 	{
 		// Only serialize if we have the requested amount of data
diff --git a/Engine/Source/Runtime/Core/Public/Serialization/MemoryReader.h b/Engine/Source/Runtime/Core/Public/Serialization/MemoryReader.h
--- a/Engine/Source/Runtime/Core/Public/Serialization/MemoryReader.h
+++ b/Engine/Source/Runtime/Core/Public/Serialization/MemoryReader.h
@@ -34,10 +34,6 @@ public:
 
 	void Serialize( void* Data, int64 Num )
 	{
-#if defined(__clang_analyzer__)
-		// Suppress uninitialized data static analysis warning
-		FMemory::Memzero(Data, Num);
-#endif
 		if (Num && !IsError())
 		{
 			// Only serialize if we have the requested amount of data
diff --git a/Engine/Source/Runtime/Core/Public/Serialization/StaticMemoryReader.h b/Engine/Source/Runtime/Core/Public/Serialization/StaticMemoryReader.h
--- a/Engine/Source/Runtime/Core/Public/Serialization/StaticMemoryReader.h
+++ b/Engine/Source/Runtime/Core/Public/Serialization/StaticMemoryReader.h
@@ -32,10 +32,6 @@ public:
 
 	void Serialize( void* Data, int64 Num )
 	{
-#if defined(__clang_analyzer__)
-		// Suppress uninitialized data static analysis warning
-		FMemory::Memzero(Data, Num);
-#endif
 		if (Num && !IsError())
 		{
 			// Only serialize if we have the requested amount of data
diff --git a/Engine/Source/Runtime/CoreUObject/Public/Serialization/ObjectReader.h b/Engine/Source/Runtime/CoreUObject/Public/Serialization/ObjectReader.h
--- a/Engine/Source/Runtime/CoreUObject/Public/Serialization/ObjectReader.h
+++ b/Engine/Source/Runtime/CoreUObject/Public/Serialization/ObjectReader.h
@@ -55,10 +55,6 @@ public:
 
 	void Serialize(void* Data, int64 Num)
 	{
-#if defined(__clang_analyzer__)
-		// Suppress uninitialized data static analysis warning
-		FMemory::Memzero(Data, Num);
-#endif
 		if (Num && !IsError())
 		{
 			// Only serialize if we have the requested amount of data
diff --git a/Engine/Source/Runtime/Net/Core/Private/Net/Core/Misc/PacketAudit.cpp b/Engine/Source/Runtime/Net/Core/Private/Net/Core/Misc/PacketAudit.cpp
--- a/Engine/Source/Runtime/Net/Core/Private/Net/Core/Misc/PacketAudit.cpp
+++ b/Engine/Source/Runtime/Net/Core/Private/Net/Core/Misc/PacketAudit.cpp
@@ -157,10 +157,6 @@ public:
 
 	virtual void Serialize(void* Data, int64 Num) override
 	{
-#if defined(__clang_analyzer__)
-		// Suppress uninitialized data static analysis warning
-		FMemory::Memzero(Data, Num);
-#endif
 		if (Num > 0 && !IsError())
 		{
 			if (Offset + Num <= TotalSize())
diff --git a/Engine/Source/Runtime/Online/BuildPatchServices/Private/Data/ManifestUObject.cpp b/Engine/Source/Runtime/Online/BuildPatchServices/Private/Data/ManifestUObject.cpp
--- a/Engine/Source/Runtime/Online/BuildPatchServices/Private/Data/ManifestUObject.cpp
+++ b/Engine/Source/Runtime/Online/BuildPatchServices/Private/Data/ManifestUObject.cpp
@@ -356,10 +356,6 @@ public:
 
 	virtual void Serialize(void* Data, int64 Num) override
 	{
-#if defined(__clang_analyzer__)
-		// Suppress uninitialized data static analysis warning
-		FMemory::Memzero(Data, Num);
-#endif
 		if (Num && !IsError())
 		{
 			// Only serialize if we have the requested amount of data
diff --git a/Engine/Source/Runtime/Online/BuildPatchServices/Private/Tests/Fake/FileSystem.fake.h b/Engine/Source/Runtime/Online/BuildPatchServices/Private/Tests/Fake/FileSystem.fake.h
--- a/Engine/Source/Runtime/Online/BuildPatchServices/Private/Tests/Fake/FileSystem.fake.h
+++ b/Engine/Source/Runtime/Online/BuildPatchServices/Private/Tests/Fake/FileSystem.fake.h
@@ -36,10 +36,6 @@ namespace BuildPatchServices
 
 		void Serialize(void* Data, int64 Num)
 		{
-#if defined(__clang_analyzer__)
-			// Suppress uninitialized data static analysis warning
-			FMemory::Memzero(Data, Num);
-#endif
 			if (Num && !IsError())
 			{
 				if (Offset + Num <= TotalSize())
