Fixed packaging failure when packaging AEC projects for platforms unsupported by Interchange.

diff --git a/Engine/Plugins/Enterprise/DataprepEditor/DataprepEditor.uplugin b/Engine/Plugins/Enterprise/DataprepEditor/DataprepEditor.uplugin
--- a/Engine/Plugins/Enterprise/DataprepEditor/DataprepEditor.uplugin
+++ b/Engine/Plugins/Enterprise/DataprepEditor/DataprepEditor.uplugin
@@ -51,7 +51,12 @@
 		},
 		{
 			"Name": "Interchange",
-			"Enabled": true
+			"Enabled": true,
+			"SupportedTargetPlatforms": [
+				"Win64",
+				"Linux",
+				"Mac"
+			]
 		}
 	],
 	"IsExperimentalVersion": false
diff --git a/Templates/TP_AEC_ArchvisBP/TP_AEC_ArchvisBP.uproject b/Templates/TP_AEC_ArchvisBP/TP_AEC_ArchvisBP.uproject
--- a/Templates/TP_AEC_ArchvisBP/TP_AEC_ArchvisBP.uproject
+++ b/Templates/TP_AEC_ArchvisBP/TP_AEC_ArchvisBP.uproject
@@ -10,7 +10,12 @@
 		},
 		{
 			"Name": "InterchangeEditor",
-			"Enabled": true
+			"Enabled": true,
+			"SupportedTargetPlatforms": [
+				"Win64",
+				"Linux",
+				"Mac"
+			]
 		},
 		{
 			"Name": "DatasmithCADImporter",
diff --git a/Templates/TP_AEC_BlankBP/TP_AEC_BlankBP.uproject b/Templates/TP_AEC_BlankBP/TP_AEC_BlankBP.uproject
--- a/Templates/TP_AEC_BlankBP/TP_AEC_BlankBP.uproject
+++ b/Templates/TP_AEC_BlankBP/TP_AEC_BlankBP.uproject
@@ -10,7 +10,12 @@
 		},
 		{
 			"Name": "InterchangeEditor",
-			"Enabled": true
+			"Enabled": true,
+			"SupportedTargetPlatforms": [
+				"Win64",
+				"Linux",
+				"Mac"
+			]
 		},
 		{
 			"Name": "DatasmithCADImporter",
diff --git a/Templates/TP_AEC_CollabBP/TP_AEC_CollabBP.uproject b/Templates/TP_AEC_CollabBP/TP_AEC_CollabBP.uproject
--- a/Templates/TP_AEC_CollabBP/TP_AEC_CollabBP.uproject
+++ b/Templates/TP_AEC_CollabBP/TP_AEC_CollabBP.uproject
@@ -41,7 +41,12 @@
 		},
 		{
 			"Name": "InterchangeEditor",
-			"Enabled": true
+			"Enabled": true,
+			"SupportedTargetPlatforms": [
+				"Win64",
+				"Linux",
+				"Mac"
+			]
 		},
 		{
 			"Name": "AxFImporter",
diff --git a/Templates/TP_MFG_CollabBP/TP_MFG_CollabBP.uproject b/Templates/TP_MFG_CollabBP/TP_MFG_CollabBP.uproject
--- a/Templates/TP_MFG_CollabBP/TP_MFG_CollabBP.uproject
+++ b/Templates/TP_MFG_CollabBP/TP_MFG_CollabBP.uproject
@@ -42,7 +42,12 @@
 		},
 		{
 			"Name": "InterchangeEditor",
-			"Enabled": true
+			"Enabled": true,
+			"SupportedTargetPlatforms": [
+				"Win64",
+				"Linux",
+				"Mac"
+			]
 		},
 		{
 			"Name": "AxFImporter",
diff --git a/Templates/TP_PhotoStudioBP/TP_PhotoStudioBP.uproject b/Templates/TP_PhotoStudioBP/TP_PhotoStudioBP.uproject
--- a/Templates/TP_PhotoStudioBP/TP_PhotoStudioBP.uproject
+++ b/Templates/TP_PhotoStudioBP/TP_PhotoStudioBP.uproject
@@ -42,7 +42,12 @@
 		},
 		{
 			"Name": "InterchangeEditor",
-			"Enabled": true
+			"Enabled": true,
+			"SupportedTargetPlatforms": [
+				"Win64",
+				"Linux",
+				"Mac"
+			]
 		},
 		{
 			"Name": "AxFImporter",
