Fixed risk of cyclical dependency in SubObjectDependencyGraph

#rnx

diff --git a/Engine/Plugins/Enterprise/DataprepEditor/Source/DataprepEditor/Private/DataprepSnapshot.cpp b/Engine/Plugins/Enterprise/DataprepEditor/Source/DataprepEditor/Private/DataprepSnapshot.cpp
--- a/Engine/Plugins/Enterprise/DataprepEditor/Source/DataprepEditor/Private/DataprepSnapshot.cpp
+++ b/Engine/Plugins/Enterprise/DataprepEditor/Source/DataprepEditor/Private/DataprepSnapshot.cpp
@@ -292,26 +292,16 @@ namespace DataprepSnapshotUtil
 				FObjectDependencyAnalyzer Analyzer( SubObject, SubObjectsSet );
 				SubObject->Serialize( Analyzer );
 
-				SubObjectDependencyGraph[SubObject].Append(Analyzer.DependentObjects);
-			}
-
-			bool bCyclingDependency = true;
-			for (auto& Entry : SubObjectDependencyGraph)
-			{
-				if (Entry.Value.Num() == 0)
+				if (!Analyzer.DependentObjects.IsEmpty())
 				{
-					bCyclingDependency = false;
-					break;
-				}
-			}
-
-			if (bCyclingDependency)
-			{
-				for (auto& Entry : SubObjectDependencyGraph)
-				{
-					if (Entry.Value.Num() == 1)
+					TSet<UObject*>& SubObjectEntry = SubObjectDependencyGraph.FindOrAdd(SubObject);
+					for (UObject* DependentObject : Analyzer.DependentObjects)
 					{
-						Entry.Value.Empty();
+						// Skip cyclical dependency
+						if (!SubObjectDependencyGraph.Contains(DependentObject))
+						{
+							SubObjectEntry.Add(DependentObject);
+						}
 					}
 				}
 			}
@@ -332,14 +322,19 @@ namespace DataprepSnapshotUtil
 
 						SubObjectsArray.Add(SubObject);
 
+						bool bNoEndlessLoop = SubObjectDependencyGraph.IsEmpty();
 						for(auto& SubEntry : SubObjectDependencyGraph)
 						{
 							if(SubEntry.Value.Num() > 0)
 							{
 								SubEntry.Value.Remove(SubObject);
 							}
+
+							bNoEndlessLoop |= SubEntry.Value.Num() == 0;
 						}
 
+						ensure(bNoEndlessLoop);
+
 						break;
 					}
 				}
