Fixed Dataprep editor crashes
- Root cause: Order of serialization then deserialization were wrong. Dependent obects were create first. Fixed.
- Root cause: Interchange import returns the editor's map as a imported object. THis will be investigated in another JIRA. Fixed by triaging assets generated by Interchange.
- Root cause: Initial implementation was caching viewport's actors used to which mesh components were attached. Fixed by removing those actors and recreating them each time the preview viewport is updated.

#rnx

diff --git a/Engine/Plugins/Enterprise/DataprepEditor/Source/DataprepCore/Private/InterchangeFileProducer.cpp b/Engine/Plugins/Enterprise/DataprepEditor/Source/DataprepCore/Private/InterchangeFileProducer.cpp
--- a/Engine/Plugins/Enterprise/DataprepEditor/Source/DataprepCore/Private/InterchangeFileProducer.cpp
+++ b/Engine/Plugins/Enterprise/DataprepEditor/Source/DataprepCore/Private/InterchangeFileProducer.cpp
@@ -246,7 +246,8 @@ bool UInterchangeFileProducer::Execute(TArray< TWeakObjectPtr< UObject > >& OutA
 	GetPipelinesFromSourceData(*SourceData, ImportAssetParameters.OverridePipelines);
 	ImportAssetParameters.OverridePipelines.Add(FSoftObjectPath(TEXT("/DataprepEditor/DefaultDataprepPipeline.DefaultDataprepPipeline")));
 
-	FInterchangeResults ImportResults = InterchangeManager.ImportSceneAsync(TransientPackage->GetPathName(), SourceData, ImportAssetParameters);
+	const FString PackagePath = TransientPackage->GetPathName();
+	FInterchangeResults ImportResults = InterchangeManager.ImportSceneAsync(PackagePath, SourceData, ImportAssetParameters);
 
 	ImportResults.Get<0>()->WaitUntilDone();
 	ImportResults.Get<1>()->WaitUntilDone();
@@ -256,17 +257,24 @@ bool UInterchangeFileProducer::Execute(TArray< TWeakObjectPtr< UObject > >& OutA
 		return false;
 	}
 
-	OutAssets.Reserve(OutAssets.Num() + ImportResults.Get<1>()->GetImportedObjects().Num());
 
-	OutAssets.Append(ImportResults.Get<0>()->GetImportedObjects());
-
-	for (UObject* Object : ImportResults.Get<1>()->GetImportedObjects())
-	{
-		if (Object && !Object->GetClass()->IsChildOf(AActor::StaticClass()))
+	TSet<UObject*> AssetsToAdd;
+	auto AddAssets = [&AssetsToAdd, &PackagePath](const TArray<UObject*>& Objects)
 		{
-			OutAssets.Add(Object);
-		}
-	}
+			for (UObject* Object : Objects)
+			{
+				if (Object && Object->GetPathName().StartsWith(PackagePath) && !Object->IsA<AActor>())
+				{
+					AssetsToAdd.Add(Object);
+				}
+			}
+		};
+
+	AddAssets(ImportResults.Get<0>()->GetImportedObjects());
+	AddAssets(ImportResults.Get<1>()->GetImportedObjects());
+
+	OutAssets.Reserve(OutAssets.Num() + AssetsToAdd.Num());
+	OutAssets.Append(AssetsToAdd.Array());
 
 	return !IsCancelled();
 }
diff --git a/Engine/Plugins/Enterprise/DataprepEditor/Source/DataprepEditor/Private/DataprepSnapshot.cpp b/Engine/Plugins/Enterprise/DataprepEditor/Source/DataprepEditor/Private/DataprepSnapshot.cpp
--- a/Engine/Plugins/Enterprise/DataprepEditor/Source/DataprepEditor/Private/DataprepSnapshot.cpp
+++ b/Engine/Plugins/Enterprise/DataprepEditor/Source/DataprepEditor/Private/DataprepSnapshot.cpp
@@ -346,7 +346,7 @@ namespace DataprepSnapshotUtil
 		MemAr << SubObjectsCount;
 
 		// Serialize class and path of each sub-object
-		for(int32 Index = SubObjectsArray.Num() - 1; Index >= 0; --Index)
+		for (int32 Index = 0; Index < SubObjectsArray.Num(); ++Index)
 		{
 			const UObject* SubObject = SubObjectsArray[Index];
 
@@ -362,7 +362,7 @@ namespace DataprepSnapshotUtil
 		// Serialize sub-objects' outer path and name
 		// Done in reverse order since a sub-object can be the outer of another sub-object
 		// it depends on. Not the opposite
-		for(int32 Index = SubObjectsArray.Num() - 1; Index >= 0; --Index)
+		for (int32 Index = 0; Index < SubObjectsArray.Num(); ++Index)
 		{
 			const UObject* SubObject = SubObjectsArray[Index];
 
@@ -416,7 +416,7 @@ namespace DataprepSnapshotUtil
 		// Create root name to avoid name collision
 		FString RootName = FGuid::NewGuid().ToString();
 
-		for(int32 Index = SubObjectsCount - 1; Index >= 0; --Index)
+		for (int32 Index = 0; Index < SubObjectsArray.Num(); ++Index)
 		{
 			FString ClassName;
 			MemAr << ClassName;
@@ -443,8 +443,8 @@ namespace DataprepSnapshotUtil
 		}
 
 		// Restore sub-objects' outer if original outer differs from Object
-		// Restoration is done in the order the serialization was done: reverse order
-		for(int32 Index = SubObjectsArray.Num() - 1; Index >= 0; --Index)
+		// Restore sub-objects' name 
+		for(int32 Index = 0; Index < SubObjectsArray.Num(); ++Index)
 		{
 			FString SoftPathString;
 			MemAr << SoftPathString;
diff --git a/Engine/Plugins/Enterprise/DataprepEditor/Source/DataprepEditor/Private/Widgets/SDataprepEditorViewport.cpp b/Engine/Plugins/Enterprise/DataprepEditor/Source/DataprepEditor/Private/Widgets/SDataprepEditorViewport.cpp
--- a/Engine/Plugins/Enterprise/DataprepEditor/Source/DataprepEditor/Private/Widgets/SDataprepEditorViewport.cpp
+++ b/Engine/Plugins/Enterprise/DataprepEditor/Source/DataprepEditor/Private/Widgets/SDataprepEditorViewport.cpp
@@ -233,7 +233,7 @@ void SDataprepEditorViewport::Construct(const FArguments& InArgs, TSharedPtr<FDa
 
 	// restore last used feature level
 	UWorld* PreviewSceneWorld = PreviewScene->GetWorld();
-	if (PreviewSceneWorld != nullptr)
+	if (ensure(PreviewSceneWorld))
 	{
 		PreviewSceneWorld->ChangeFeatureLevel(GWorld->GetFeatureLevel());
 	}
@@ -245,77 +245,45 @@ void SDataprepEditorViewport::Construct(const FArguments& InArgs, TSharedPtr<FDa
 		PreviewSceneWorld->ChangeFeatureLevel(NewFeatureLevel);
 	});
 
-	TFunction<AActor* (USceneComponent*, const TCHAR*)> CreateActor;
-	CreateActor = [&PreviewSceneWorld](USceneComponent* ParentComponent, const TCHAR* ActorName)-> AActor*
-		{
-			AActor* Actor = PreviewSceneWorld->SpawnActor<AActor>( AActor::StaticClass(), FTransform::Identity );
-
-			if ( Actor->GetRootComponent() == nullptr )
-			{
-				USceneComponent* RootComponent = NewObject< USceneComponent >( Actor, USceneComponent::StaticClass(), ActorName, RF_NoFlags );
-				Actor->SetRootComponent( RootComponent );
-			}
-
-			if(ParentComponent)
-			{
-				Actor->GetRootComponent()->AttachToComponent( ParentComponent, FAttachmentTransformRules::KeepRelativeTransform );
-			}
-
-			Actor->RegisterAllComponents();
-
-			return Actor;
-		};
-
 	// Create root actor of 3D viewport's world
 	AActor* RootActor = CreateActor(nullptr, TEXT("RootPreviewActor"));
 	RootActorComponent = RootActor->GetRootComponent();
 
-
-	// Create actor to hold all 'Static' preview mesh components
-	StaticActor = CreateActor(RootActorComponent.Get(), TEXT("StaticPreviewActor"));
-	StaticActor->GetRootComponent()->SetMobility(EComponentMobility::Static);
-
-	// If HW supports it, create actor to hold all 'Movable' preview mesh components
-	if(GEditor->PreviewPlatform.GetEffectivePreviewFeatureLevel() > ERHIFeatureLevel::ES3_1)
-	{
-		MovableActor = CreateActor(RootActorComponent.Get(), TEXT("MovablePreviewActor"));
-		MovableActor->GetRootComponent()->SetMobility(EComponentMobility::Movable);
-	}
-
 	WorldToPreview = InArgs._WorldToPreview;
-	check(WorldToPreview);
+	ensure(WorldToPreview);
 
 	SEditorViewport::Construct( SEditorViewport::FArguments() );
 }
 
 void SDataprepEditorViewport::ClearScene()
 {
-	const int32 PreviousCount = PreviewMeshComponents.Num();
-	
-	FDetachmentTransformRules DetachmentRules(EDetachmentRule::KeepRelative, true);
+	if (!StaticActor.IsValid())
+	{
+		ensure(PreviewMeshComponents.IsEmpty() && MeshComponentsMapping.IsEmpty() && MeshComponentsReverseMapping.IsEmpty() && SelectedPreviewComponents.IsEmpty());
+		return;
+	}
 
-	StaticActor->UnregisterAllComponents();
-	MovableActor->UnregisterAllComponents();
-	
-	TArray<UObject*> ObjectsToDelete;
-	ObjectsToDelete.Reserve(PreviousCount);
+	UWorld* PreviewSceneWorld = PreviewScene->GetWorld();
 
-	for( const TWeakObjectPtr< UStaticMeshComponent >& PreviewMeshComponent : PreviewMeshComponents )
-	{
-		if(UStaticMeshComponent* MeshComponent = PreviewMeshComponent.Get())
+	TFunction<void(AActor*)> DeleteActor = [PreviewSceneWorld](AActor* ActorToDelete)
 		{
-			MeshComponent->DetachFromComponent(DetachmentRules);
-			MeshComponent->SetStaticMesh(nullptr);
+			if (IsValid(ActorToDelete))
+			{
+				PreviewSceneWorld->EditorDestroyActor(ActorToDelete, true);
 
-			MeshComponent->EmptyOverrideMaterials();
-			ObjectsToDelete.Add( MeshComponent );
-		}
-	}
+				// Since deletion can be delayed, rename to avoid future name collision
+				// Call UObject::Rename directly on actor to avoid AActor::Rename which unnecessarily sunregister and re-register components
+				ActorToDelete->UObject::Rename(nullptr, GetTransientPackage(), REN_DontCreateRedirectors);
+			}
+		};
 
-	StaticActor->RegisterAllComponents();
-	MovableActor->RegisterAllComponents();
+	DeleteActor(StaticActor.Get());
+	DeleteActor(MovableActor.Get());
 
-	FDataprepCoreUtils::PurgeObjects( MoveTemp( ObjectsToDelete ) );
+	StaticActor = nullptr;
+	MovableActor = nullptr;
+
+	const int32 PreviousCount = PreviewMeshComponents.Num();
 
 	PreviewMeshComponents.Empty(PreviousCount);
 	MeshComponentsMapping.Empty(PreviousCount);
@@ -331,6 +299,18 @@ void SDataprepEditorViewport::UpdateScene()
 
 	ClearScene();
 
+	// Create actor to hold all 'Static' preview mesh components
+	StaticActor = CreateActor(RootActorComponent.Get(), TEXT("StaticPreviewActor"));
+	StaticActor->GetRootComponent()->SetMobility(EComponentMobility::Static);
+	ensure(StaticActor.IsValid());
+
+	// If HW supports it, create actor to hold all 'Movable' preview mesh components
+	if (GEditor->PreviewPlatform.GetEffectivePreviewFeatureLevel() > ERHIFeatureLevel::ES3_1)
+	{
+		MovableActor = CreateActor(RootActorComponent.Get(), TEXT("MovablePreviewActor"));
+		MovableActor->GetRootComponent()->SetMobility(EComponentMobility::Movable);
+	}
+
 	ViewportDebug::FTimeLogger TimeLogger( TEXT("Updating viewport") );
 
 	SceneBounds = FBox( FVector::ZeroVector, FVector(100.0f) );
@@ -436,7 +416,7 @@ void SDataprepEditorViewport::UpdateScene()
 							Mobility = SceneMeshComponent->Mobility;
 						}
 
-						AActor* PreviewActor = Mobility == EComponentMobility::Movable ? MovableActor.Get() : StaticActor.Get();
+						AActor* PreviewActor = Mobility == EComponentMobility::Movable && MovableActor.IsValid() ? MovableActor.Get() : StaticActor.Get();
 
 						UCustomStaticMeshComponent* PreviewMeshComponent = NewObject< UCustomStaticMeshComponent >( PreviewActor, NAME_None, RF_Transient );
 						PreviewMeshComponent->SetMobility( Mobility );
@@ -1602,6 +1582,28 @@ TSharedRef<SWidget> SDataprepEditorViewportToolbar::GenerateRenderingMenu() cons
 	return MenuBuilder.MakeWidget();
 }
 
+AActor* SDataprepEditorViewport::CreateActor(USceneComponent* ParentComponent, const TCHAR* ActorName)
+{
+	UWorld* PreviewSceneWorld = PreviewScene->GetWorld();
+
+	AActor* Actor = PreviewSceneWorld->SpawnActor<AActor>(AActor::StaticClass(), FTransform::Identity);
+
+	if (Actor->GetRootComponent() == nullptr)
+	{
+		USceneComponent* RootComponent = NewObject< USceneComponent >(Actor, USceneComponent::StaticClass(), ActorName, RF_NoFlags);
+		Actor->SetRootComponent(RootComponent);
+	}
+
+	if (ParentComponent)
+	{
+		Actor->GetRootComponent()->AttachToComponent(ParentComponent, FAttachmentTransformRules::KeepRelativeTransform);
+	}
+
+	Actor->RegisterAllComponents();
+
+	return Actor;
+}
+
 #ifdef VIEWPORT_EXPERIMENTAL
 TSharedRef<SWidget> SDataprepEditorViewportToolbar::GenerateExperimentalMenu() const
 {
diff --git a/Engine/Plugins/Enterprise/DataprepEditor/Source/DataprepEditor/Private/Widgets/SDataprepEditorViewport.h b/Engine/Plugins/Enterprise/DataprepEditor/Source/DataprepEditor/Private/Widgets/SDataprepEditorViewport.h
--- a/Engine/Plugins/Enterprise/DataprepEditor/Source/DataprepEditor/Private/Widgets/SDataprepEditorViewport.h
+++ b/Engine/Plugins/Enterprise/DataprepEditor/Source/DataprepEditor/Private/Widgets/SDataprepEditorViewport.h
@@ -315,6 +315,8 @@ private:
 
 	/** Callback for checking whether the Nanite show entry should be displayed. */
 	bool IsShowNaniteFallbackVisible() const { return bCanShowNaniteFallbackMenu; }
+	
+	AActor* CreateActor(USceneComponent* ParentComponent, const TCHAR* ActorName);
 
 private:
 	/** The scene for this viewport */
