Fixed issue with output folder path containing '/' character
- Reported by a user on UDN

#rnx

diff --git a/Engine/Plugins/Enterprise/DataprepEditor/Source/DataprepLibraries/Private/DataprepOperations.cpp b/Engine/Plugins/Enterprise/DataprepEditor/Source/DataprepLibraries/Private/DataprepOperations.cpp
--- a/Engine/Plugins/Enterprise/DataprepEditor/Source/DataprepLibraries/Private/DataprepOperations.cpp
+++ b/Engine/Plugins/Enterprise/DataprepEditor/Source/DataprepLibraries/Private/DataprepOperations.cpp
@@ -462,30 +462,41 @@ void FDataprepSetOutputFolderDetails::FolderName_TextCommited(const FText& InTex
 void FDataprepSetOutputFolderDetails::FolderName_TextChanged(const FText& Text)
 {
 	// Slash and Square brackets are invalid characters for a folder name
-	const FString InvalidChars = INVALID_LONGPACKAGE_CHARACTERS TEXT("/[]"); 
+	const FString InvalidChars = INVALID_OBJECTPATH_CHARACTERS TEXT(".");
 
 	FText ErrorMessage;
-	FString FolderName = Text.ToString();
+	FString FolderPath = Text.ToString();
 
-	// See if the name contains invalid characters.
-	FString Char;
-	for( int32 CharIdx = 0; CharIdx < FolderName.Len(); ++CharIdx )
+	if (FolderPath.Len() == 0)
 	{
-		Char = FolderName.Mid(CharIdx, 1);
-
-		if ( InvalidChars.Contains(*Char) )
+		ErrorMessage = LOCTEXT("InvalidFolderName_PathEmpty", "A folder path must be at least one character long");
+	}
+	else if (FolderPath.StartsWith(TEXT("/")))
+	{
+		ErrorMessage = LOCTEXT("InvalidFolderName_MustBeRelative", "A folder path must be relative");
+	}
+	else
+	{
+		// See if the name contains invalid characters.
+		FString Char;
+		for (int32 CharIdx = 0; CharIdx < FolderPath.Len(); ++CharIdx)
 		{
-			FString ReadableInvalidChars = InvalidChars;
-			ReadableInvalidChars.ReplaceInline(TEXT("\r"), TEXT(""));
-			ReadableInvalidChars.ReplaceInline(TEXT("\n"), TEXT(""));
-			ReadableInvalidChars.ReplaceInline(TEXT("\t"), TEXT(""));
+			Char = FolderPath.Mid(CharIdx, 1);
+
+			if (InvalidChars.Contains(*Char))
+			{
+				FString ReadableInvalidChars = InvalidChars;
+				ReadableInvalidChars.ReplaceInline(TEXT("\r"), TEXT(""));
+				ReadableInvalidChars.ReplaceInline(TEXT("\n"), TEXT(""));
+				ReadableInvalidChars.ReplaceInline(TEXT("\t"), TEXT(""));
 
-			ErrorMessage = FText::Format(LOCTEXT("InvalidFolderName_InvalidCharacters", "A folder name may not contain any of the following characters: {0}"), FText::FromString(ReadableInvalidChars));
-			break;
+				ErrorMessage = FText::Format(LOCTEXT("InvalidFolderName_InvalidCharacters", "A folder path may not contain any of the following characters: {0}"), FText::FromString(ReadableInvalidChars));
+				break;
+			}
 		}
 	}
 
-	if (!ErrorMessage.IsEmpty() || !FFileHelper::IsFilenameValidForSaving(FolderName, ErrorMessage))
+	if (!ErrorMessage.IsEmpty() || !FFileHelper::IsFilenameValidForSaving(FolderPath, ErrorMessage))
 	{
 		TextBox->SetError(ErrorMessage);
 	}
