Fixed Dataprep's transient root package path following changes to policy on mount points

#rrnx

#Synced-CL 50317384

diff --git a/Engine/Plugins/Enterprise/DataprepEditor/Source/DataprepCore/Private/DataprepCorePrivateUtils.cpp b/Engine/Plugins/Enterprise/DataprepEditor/Source/DataprepCore/Private/DataprepCorePrivateUtils.cpp
--- a/Engine/Plugins/Enterprise/DataprepEditor/Source/DataprepCore/Private/DataprepCorePrivateUtils.cpp
+++ b/Engine/Plugins/Enterprise/DataprepEditor/Source/DataprepCore/Private/DataprepCorePrivateUtils.cpp
@@ -45,7 +45,7 @@ const FString& DataprepCorePrivateUtils::GetRootTemporaryDir()
 
 const FString& DataprepCorePrivateUtils::GetRootPackagePath()
 {
-	static FString RootPackagePath( TEXT("/DataprepCore/Transient") );
+	static FString RootPackagePath( TEXT("/DataprepTransient") );
 	return RootPackagePath;
 }
 
