Made the Interchange FIleProducer available by default

#rnx
#okforgithhub public

#Synced-CL 51338605

diff --git a/Engine/Plugins/Enterprise/DataprepEditor/Source/DataprepCore/Private/InterchangeFileProducer.cpp b/Engine/Plugins/Enterprise/DataprepEditor/Source/DataprepCore/Private/InterchangeFileProducer.cpp
--- a/Engine/Plugins/Enterprise/DataprepEditor/Source/DataprepCore/Private/InterchangeFileProducer.cpp
+++ b/Engine/Plugins/Enterprise/DataprepEditor/Source/DataprepCore/Private/InterchangeFileProducer.cpp
@@ -37,14 +37,14 @@
 
 #define LOCTEXT_NAMESPACE "InterchangeFileProducer"
 
-const FText InterchangeFileProducerLabel( LOCTEXT( "InterchangeFileProducerLabel", "Interchange file importer" ) );
-const FText InterchangeFileProducerDescription( LOCTEXT( "InterchangeFileProducerDesc", "Reads a file supported by Interchange and its dependent assets" ) );
+const FText InterchangeFileProducerLabel( LOCTEXT( "InterchangeFileProducerLabel", "Interchange file producer" ) );
+const FText InterchangeFileProducerDescription( LOCTEXT( "InterchangeFileProducerDesc", "Reads files of formats supported by Interchange, i.e. FBX, glTF, ..." ) );
 
-bool bGEnableInterchangeProducer = false;
+bool bGEnableInterchangeProducer = true;
 FAutoConsoleVariableRef GEnableInterchangeProducerVar(
 	TEXT("Dataprep.Interchange.Producer"),
 	bGEnableInterchangeProducer,
-	TEXT("Enable/disable the Intrchange file producer.\nDefault is disabled\n"),
+	TEXT("Enable/disable the Interchange file producer.\nDefault is enabled\n"),
 	ECVF_Default
 );
 
